Add Oodle integration support

diff --git a/Engine/Plugins/Runtime/PacketHandlers/CompressionComponents/Oodle/Oodle.uplugin b/Engine/Plugins/Runtime/PacketHandlers/CompressionComponents/Oodle/Oodle.uplugin
new file mode 100644
--- /dev/null
+++ b/Engine/Plugins/Runtime/PacketHandlers/CompressionComponents/Oodle/Oodle.uplugin
@@ -0,0 +1,23 @@
+{
+	"FileVersion" : 3,
+	
+	"FriendlyName" : "Oddle network compression",
+	"Version" : 1,
+	"VersionName" : "1.0",
+	"CreatedBy" : "Epic Games, Inc.",
+	"CreatedByURL" : "http://epicgames.com",
+	"EngineVersion" : "4.4.0",
+	"Description" : "Realtime network compression using Oodle technology",
+	"Category" : "Networking",
+	"EnabledByDefault" : false,
+    "IsBetaVersion": false,
+	
+	"Modules" :
+	[
+		{
+			"Name" : "OodleHandlerComponent",
+			"Type" : "Runtime",
+			"WhitelistPlatforms" : [ "Win32", "Win64" ]
+		}
+	]
+}
\ No newline at end of file
diff --git a/Engine/Plugins/Runtime/PacketHandlers/CompressionComponents/Oodle/Source/OodleHandlerComponent/OodleHandlerComponent.Build.cs b/Engine/Plugins/Runtime/PacketHandlers/CompressionComponents/Oodle/Source/OodleHandlerComponent/OodleHandlerComponent.Build.cs
new file mode 100644
--- /dev/null
+++ b/Engine/Plugins/Runtime/PacketHandlers/CompressionComponents/Oodle/Source/OodleHandlerComponent/OodleHandlerComponent.Build.cs
@@ -0,0 +1,40 @@
+// Copyright 1998-2015 Epic Games, Inc. All Rights Reserved.
+
+using UnrealBuildTool;
+using System.IO;
+
+public class OodleHandlerComponent : ModuleRules
+{
+    public OodleHandlerComponent(TargetInfo Target)
+    {
+		BinariesSubFolder = "NotForLicensees";
+		
+		PublicDependencyModuleNames.AddRange( new string[] { "PacketHandler", "Core" } );
+
+		bool bHaveOodleSDK = false;
+
+		if ( ( Target.Platform == UnrealTargetPlatform.Win64 ) || ( Target.Platform == UnrealTargetPlatform.Win32 ) )
+        {
+			string OodleNotForLicenseesLibDir = System.IO.Path.Combine( UEBuildConfiguration.UEThirdPartySourceDirectory, "..", "..", "Plugins", "Runtime", "PacketHandlers", "CompressionComponents", "Oodle", "Source", "ThirdParty", "NotForLicensees", "Oodle", "win", "lib" );   // Check the NotForLicensees folder first
+
+			try
+			{
+				bHaveOodleSDK = System.IO.Directory.Exists( OodleNotForLicenseesLibDir );
+			}
+			catch ( System.Exception )
+			{
+			}
+        }
+
+		if ( bHaveOodleSDK )
+		{
+	        AddThirdPartyPrivateStaticDependencies(Target,"Oodle");
+	        PublicIncludePathModuleNames.Add("Oodle");
+			Definitions.Add( "HAS_OODLE_SDK=1" );
+		}
+		else
+		{
+			Definitions.Add( "HAS_OODLE_SDK=0" );
+		}
+    }
+}
diff --git a/Engine/Plugins/Runtime/PacketHandlers/CompressionComponents/Oodle/Source/OodleHandlerComponent/Private/OodleHandlerComponent.cpp b/Engine/Plugins/Runtime/PacketHandlers/CompressionComponents/Oodle/Source/OodleHandlerComponent/Private/OodleHandlerComponent.cpp
new file mode 100644
--- /dev/null
+++ b/Engine/Plugins/Runtime/PacketHandlers/CompressionComponents/Oodle/Source/OodleHandlerComponent/Private/OodleHandlerComponent.cpp
@@ -0,0 +1,315 @@
+// Copyright 1998-2015 Epic Games, Inc. All Rights Reserved.
+
+#include "OodleHandlerComponent.h"
+
+#if HAS_OODLE_SDK == 1
+IMPLEMENT_MODULE( FOodleComponentModuleInterface, OodleHandlerComponent );
+
+DEFINE_LOG_CATEGORY(OodleHandlerComponentLog);
+
+DEFINE_STAT(STAT_Oodle_BytesRaw);
+DEFINE_STAT(STAT_Oodle_BytesCompressed);
+
+// Saved compressor model file :
+// must match example_packet - or whatever you use to write the state file
+struct OodleNetwork1_SavedStates_Header
+{
+#define ON1_MAGIC	0x11235801
+	U32 magic;
+	U32 compressor;
+	U32 ht_bits;
+	U32 dic_size;
+	U32 oodle_header_version;
+	U32 dic_complen;
+	U32 statecompacted_size;
+	U32 statecompacted_complen;
+};
+
+// OODLE
+OodleHandlerComponent::OodleHandlerComponent()
+: BytesSaved(0)
+, LogFile(nullptr)
+, SharedDictionaryData(nullptr)
+, CompressorState(nullptr)
+, Dictionary(nullptr)
+{
+	SetActive(true);
+}
+
+OodleHandlerComponent::~OodleHandlerComponent()
+{
+	free(Dictionary);
+	free(SharedDictionaryData);
+	free(CompressorState);
+
+	Oodle_Shutdown_NoThreads();
+}
+
+void OodleHandlerComponent::Initialize()
+{
+	// Mode
+	FString ReadMode;
+	GConfig->GetString(TEXT("OodleHandlerComponent"), TEXT("Mode"), ReadMode, GEngineIni);
+
+	if(ReadMode == TEXT("Training"))
+	{
+		Mode = Training;
+	}
+	else if(ReadMode == TEXT("Release"))
+	{
+		Mode = Release;
+	}
+	// Default
+	else
+	{
+		Mode = Release;
+	}
+
+	switch(Mode)
+	{
+		case Training:
+		{
+			if(Handler->Mode == Handler::Mode::Server)
+			{
+				FString ReadOutputLogDirectory;
+				GConfig->GetString(TEXT("OodleHandlerComponent"), TEXT("LogFile"), ReadOutputLogDirectory, GEngineIni);
+
+				LogFile = FPlatformFileManager::Get().GetPlatformFile().OpenWrite(*ReadOutputLogDirectory);
+
+				if(!LogFile)
+				{
+					LowLevelFatalError(TEXT("Failed to create file %s"), *ReadOutputLogDirectory);
+				}
+
+				int32 NumberOfChannels = 1;
+			
+				LogFile->Write(reinterpret_cast<const uint8*>(&NumberOfChannels), sizeof(NumberOfChannels));
+			}
+			break;
+		}
+		case Release:
+		{
+			FString ReadStateDirectory;
+			GConfig->GetString(TEXT("OodleHandlerComponent"), TEXT("StateFile"), ReadStateDirectory, GEngineIni);
+
+			if(ReadStateDirectory.Len() < 0)
+			{
+				LowLevelFatalError(TEXT("Specify a state directory for Oodle compressor in DefaultEngine.ini"));
+				return;
+			}
+
+			TArray<uint8> FileInMemory;
+			if (!FFileHelper::LoadFileToArray(FileInMemory, *ReadStateDirectory))
+			{
+				LowLevelFatalError(TEXT("Incorrect LogFile Provided"));
+			}
+	
+			// parse header :
+			OodleNetwork1_SavedStates_Header* SavedStateFileData = (OodleNetwork1_SavedStates_Header *)FileInMemory.GetData();
+
+			if (SavedStateFileData->magic != ON1_MAGIC)
+			{
+				UE_LOG(OodleHandlerComponentLog, Warning, TEXT("OodleNetCompressor: state_file ON1_MAGIC mismatch"));
+				return;
+			}
+
+			if (SavedStateFileData->oodle_header_version != OODLE_HEADER_VERSION)
+			{
+				UE_LOG(OodleHandlerComponentLog, Warning, TEXT("OodleNetCompressor: state_file OODLE_HEADER_VERSION mismatch. Was: %u, Expected: %u"), (uint32)SavedStateFileData->oodle_header_version, (uint32)OODLE_HEADER_VERSION);
+				return;
+			}
+
+			// this code only loads compressed states currently :
+			check(SavedStateFileData->compressor != -1);
+
+			// Init
+			OodleInitOptions Options;
+			//Oodle_Init_GetDefaults(OODLE_HEADER_VERSION, &Options);
+			Oodle_Init_GetDefaults_Minimal(OODLE_HEADER_VERSION, &Options);
+			Oodle_Init_NoThreads(OODLE_HEADER_VERSION, &Options);
+
+			S32 StateHTBits = SavedStateFileData->ht_bits;
+			SINTa DictionarySize = SavedStateFileData->dic_size;
+
+			SINTa DictionaryCompLen = SavedStateFileData->dic_complen;
+			SINTa CompressorStateCompactedSize = SavedStateFileData->statecompacted_size;
+			SINTa CompressorStateCompactedLength = SavedStateFileData->statecompacted_complen;
+
+			SINTa SharedSize = OodleNetwork1_Shared_Size(StateHTBits);
+			SINTa CompressorStateSize = OodleNetwork1UDP_State_Size();
+
+			check(DictionarySize >= DictionaryCompLen);
+			check(CompressorStateCompactedSize >= CompressorStateCompactedLength);
+
+			check(CompressorStateCompactedSize > 0 && CompressorStateCompactedSize < OodleNetwork1UDP_StateCompacted_MaxSize());
+
+			check(FileInMemory.Num() == (S64)sizeof(OodleNetwork1_SavedStates_Header)+DictionaryCompLen + CompressorStateCompactedLength);
+
+			//OodleLog_Printf_v1("OodleNetwork1UDP Loading; dic comp %d , state %d->%d\n",DictionaryCompLen, CompressorStateCompactedSize, CompressorStateCompactedLength);
+
+			//-------------------------------------------
+			// decompress Dictionary and CompressorStatecompacted
+
+			Dictionary = malloc(DictionarySize);
+
+			void * DictionaryCompressed = SavedStateFileData + 1;
+
+			SINTa DecompressedDictionarySize = OodleLZ_Decompress(DictionaryCompressed, DictionaryCompLen, Dictionary, DictionarySize);
+
+			check(DecompressedDictionarySize == DictionarySize);
+
+			OodleNetwork1UDP_StateCompacted* UDPNewCompacted = (OodleNetwork1UDP_StateCompacted *)malloc(CompressorStateCompactedSize);
+
+			void * CompressorStatecompacted_comp_ptr = (U8 *)DictionaryCompressed + DictionaryCompLen;
+
+			SINTa decomp_statecompacted_size = OodleLZ_Decompress(CompressorStatecompacted_comp_ptr, CompressorStateCompactedLength, UDPNewCompacted, CompressorStateCompactedSize);
+
+			check(decomp_statecompacted_size == CompressorStateCompactedSize);
+
+			//----------------------------------------------
+			// Uncompact the "Compacted" state into a usable state
+
+			CompressorState = (OodleNetwork1UDP_State *)malloc(CompressorStateSize);
+			check(CompressorState != NULL);
+
+			OodleNetwork1UDP_State_Uncompact(CompressorState, UDPNewCompacted);
+			free(UDPNewCompacted);
+
+			//----------------------------------------------
+			// fill out SharedDictionaryData from the dictionary
+
+			SharedDictionaryData = (OodleNetwork1_Shared *)malloc(SharedSize);
+			check(SharedDictionaryData != NULL);
+			OodleNetwork1_Shared_SetWindow(SharedDictionaryData, StateHTBits, Dictionary, (S32)DictionarySize);
+
+			//-----------------------------------------------------
+	
+			//SET_DWORD_STAT(STAT_Oodle_DicSize, DictionarySize);
+			break;
+		}
+		default:
+		break;
+	}
+
+	Initialized();
+}
+
+bool OodleHandlerComponent::IsValid() const
+{
+	return true;
+}
+
+void OodleHandlerComponent::Incoming(FBitReader& Packet)
+{
+	switch(Mode)
+	{
+		case Training:
+		{
+			if(Handler->Mode == Handler::Mode::Server)
+			{
+				uint32 SizeOfPacket = Packet.GetNumBytes();
+				int32 channel = 0;
+
+				LogFile->Write(reinterpret_cast<const uint8*>(&channel), sizeof(channel));
+				LogFile->Write(reinterpret_cast<const uint8*>(&SizeOfPacket), sizeof(SizeOfPacket));
+				LogFile->Write(reinterpret_cast<const uint8*>(Packet.GetData()), SizeOfPacket);
+			}
+			break;
+		}
+		case Release:
+		{
+			// Remove size
+			uint32 DecompressedLength;
+			Packet.SerializeIntPacked(DecompressedLength);
+
+			const int32 CompressedLength = Packet.GetBytesLeft();
+			uint8* CompressedData = new uint8[CompressedLength];
+			Packet.Serialize(CompressedData, CompressedLength);
+
+			uint8* DecompressedData = new uint8[DecompressedLength];
+	
+			check(OodleNetwork1UDP_Decode(CompressorState, SharedDictionaryData, CompressedData, CompressedLength, DecompressedData, DecompressedLength));
+
+			FBitReader UnCompressedPacket(DecompressedData, DecompressedLength * 8);
+			Packet = UnCompressedPacket;
+
+			delete[] CompressedData;
+			delete[] DecompressedData;
+			break;
+		}
+	}
+}
+
+void OodleHandlerComponent::Outgoing(FBitWriter& Packet)
+{
+	switch(Mode)
+	{
+		case Training:
+		{
+			if(Handler->Mode == Handler::Mode::Server)
+			{
+				uint32 SizeOfPacket = Packet.GetNumBytes();
+				int32 channel = 0;
+				LogFile->Write(reinterpret_cast<const uint8*>(&channel), sizeof(channel));
+				LogFile->Write(reinterpret_cast<const uint8*>(&SizeOfPacket), sizeof(SizeOfPacket));
+				LogFile->Write(reinterpret_cast<const uint8*>(Packet.GetData()), SizeOfPacket);
+			}
+			break;
+		}
+		case Release:
+		{
+			// Add size
+			uint32 UncompressedLength = Packet.GetNumBytes();
+
+			if(UncompressedLength > 0)
+			{
+				uint8* UncompressedData = new uint8[UncompressedLength];
+				memcpy(UncompressedData, Packet.GetData(), UncompressedLength);
+
+				uint8_t* CompressedData = new uint8_t[OodleLZ_GetCompressedBufferSizeNeeded(UncompressedLength)];
+
+				SINTa CompressedLengthSINT = OodleNetwork1UDP_Encode(CompressorState, SharedDictionaryData, UncompressedData, UncompressedLength, CompressedData);
+				uint32 CompressedLength = static_cast<uint32>(CompressedLengthSINT);
+
+				Packet.Reset();
+
+				Packet.SerializeIntPacked(UncompressedLength);
+				Packet.Serialize(CompressedData, CompressedLength);
+
+				delete UncompressedData;
+				delete CompressedData;
+
+				BytesSaved += UncompressedLength - CompressedLength;
+
+				//UE_LOG(OodleHandlerComponentLog, Log, TEXT("Oodle Compressed: UnCompressed: %i Compressed: %i Total: %i"), UncompressedLength, CompressedLength, BytesSaved);
+
+				INC_DWORD_STAT_BY(STAT_Oodle_BytesRaw, UncompressedLength);
+				INC_DWORD_STAT_BY(STAT_Oodle_BytesCompressed, CompressedLength);
+			}
+			else
+			{
+				Packet.SerializeIntPacked(UncompressedLength);
+			}
+
+			break;
+		}
+	}
+}
+
+// MODULE INTERFACE
+HandlerComponent* FOodleComponentModuleInterface::CreateComponentInstance()
+{
+	return new OodleHandlerComponent;
+}
+
+void FOodleComponentModuleInterface::StartupModule()
+{
+	Oodle_Init_Default(OODLE_HEADER_VERSION,
+		Oodle_Init_GetDefaults_DebugSystems_No,Oodle_Init_GetDefaults_Threads_No);
+}
+
+void FOodleComponentModuleInterface::ShutdownModule()
+{
+	Oodle_Shutdown();
+}
+#endif
\ No newline at end of file
diff --git a/Engine/Plugins/Runtime/PacketHandlers/CompressionComponents/Oodle/Source/OodleHandlerComponent/Public/OodleHandlerComponent.h b/Engine/Plugins/Runtime/PacketHandlers/CompressionComponents/Oodle/Source/OodleHandlerComponent/Public/OodleHandlerComponent.h
new file mode 100644
--- /dev/null
+++ b/Engine/Plugins/Runtime/PacketHandlers/CompressionComponents/Oodle/Source/OodleHandlerComponent/Public/OodleHandlerComponent.h
@@ -0,0 +1,72 @@
+// Copyright 1998-2015 Epic Games, Inc. All Rights Reserved.
+
+#pragma once
+
+#include "PacketHandler.h"
+#include "ModuleManager.h"
+#include "Core.h"
+
+#if HAS_OODLE_SDK == 1
+#include "oodle.h"
+
+DECLARE_LOG_CATEGORY_EXTERN(OodleHandlerComponentLog, Log, All);
+
+// Symmetric Stream cipher
+class OODLEHANDLERCOMPONENT_API OodleHandlerComponent : public HandlerComponent
+{
+public:
+	/* Initializes default data */
+	OodleHandlerComponent();
+
+	/* Default Destructor */
+	~OodleHandlerComponent();
+
+	/* Initializes the handler component */
+	virtual void Initialize() override;
+
+	/* Whether the handler component is valid */
+	virtual bool IsValid() const override;
+
+	/* Handles any incoming packets */
+	virtual void Incoming(FBitReader& Packet) override;
+
+	/* Handles any outgoing packets */
+	virtual void Outgoing(FBitWriter& Packet) override;
+
+protected:
+	/* How many bytes have been saved from compressing the data */
+	uint32 BytesSaved;
+
+	/* File to log packets to */
+	IFileHandle* LogFile;
+
+	/* 
+	* Modes of the component
+	* Specify this in .ini by
+	* [OodleHandlerComponent]
+	* Mode=Training
+	*/
+	enum {Training, Release}Mode;
+
+	/* Dictionary */
+	void* Dictionary; 
+
+	/* Shared Compressor */
+	OodleNetwork1_Shared* SharedDictionaryData;
+
+	/* State of Oodle */
+	OodleNetwork1UDP_State* CompressorState;
+};
+
+/* Reliability Module Interface */
+class FOodleComponentModuleInterface : public FPacketHandlerComponentModuleInterface
+{
+public:
+	virtual HandlerComponent* CreateComponentInstance() override;
+	virtual void StartupModule() override;
+	virtual void ShutdownModule() override;
+};
+
+DECLARE_DWORD_COUNTER_STAT_EXTERN(TEXT("Oodle Bytes Raw"), STAT_Oodle_BytesRaw, STATGROUP_Net, );
+DECLARE_DWORD_COUNTER_STAT_EXTERN(TEXT("Oodle Bytes Compressed"), STAT_Oodle_BytesCompressed, STATGROUP_Net, );
+#endif
\ No newline at end of file
