Fix issue where OodleHandlerComponent plug-in fails to load when Oodle sdk is missing

diff --git a/Engine/Plugins/Runtime/PacketHandlers/CompressionComponents/Oodle/Source/OodleHandlerComponent/Private/OodleHandlerComponent.cpp b/Engine/Plugins/Runtime/PacketHandlers/CompressionComponents/Oodle/Source/OodleHandlerComponent/Private/OodleHandlerComponent.cpp
--- a/Engine/Plugins/Runtime/PacketHandlers/CompressionComponents/Oodle/Source/OodleHandlerComponent/Private/OodleHandlerComponent.cpp
+++ b/Engine/Plugins/Runtime/PacketHandlers/CompressionComponents/Oodle/Source/OodleHandlerComponent/Private/OodleHandlerComponent.cpp
@@ -304,12 +304,26 @@ HandlerComponent* FOodleComponentModuleInterface::CreateComponentInstance()
 
 void FOodleComponentModuleInterface::StartupModule()
 {
-	Oodle_Init_Default(OODLE_HEADER_VERSION,
-		Oodle_Init_GetDefaults_DebugSystems_No,Oodle_Init_GetDefaults_Threads_No);
+	Oodle_Init_Default(OODLE_HEADER_VERSION, Oodle_Init_GetDefaults_DebugSystems_No,Oodle_Init_GetDefaults_Threads_No);
 }
 
 void FOodleComponentModuleInterface::ShutdownModule()
 {
 	Oodle_Shutdown();
 }
+#else
+IMPLEMENT_MODULE( FOodleComponentModuleInterface, OodleHandlerComponent );
+
+HandlerComponent* FOodleComponentModuleInterface::CreateComponentInstance()
+{
+	return nullptr;
+}
+
+void FOodleComponentModuleInterface::StartupModule()
+{
+}
+
+void FOodleComponentModuleInterface::ShutdownModule()
+{
+}
 #endif
\ No newline at end of file
diff --git a/Engine/Plugins/Runtime/PacketHandlers/CompressionComponents/Oodle/Source/OodleHandlerComponent/Public/OodleHandlerComponent.h b/Engine/Plugins/Runtime/PacketHandlers/CompressionComponents/Oodle/Source/OodleHandlerComponent/Public/OodleHandlerComponent.h
--- a/Engine/Plugins/Runtime/PacketHandlers/CompressionComponents/Oodle/Source/OodleHandlerComponent/Public/OodleHandlerComponent.h
+++ b/Engine/Plugins/Runtime/PacketHandlers/CompressionComponents/Oodle/Source/OodleHandlerComponent/Public/OodleHandlerComponent.h
@@ -58,6 +58,10 @@ protected:
 	OodleNetwork1UDP_State* CompressorState;
 };
 
+DECLARE_DWORD_COUNTER_STAT_EXTERN( TEXT( "Oodle Bytes Raw" ), STAT_Oodle_BytesRaw, STATGROUP_Net, );
+DECLARE_DWORD_COUNTER_STAT_EXTERN( TEXT( "Oodle Bytes Compressed" ), STAT_Oodle_BytesCompressed, STATGROUP_Net, );
+#endif
+
 /* Reliability Module Interface */
 class FOodleComponentModuleInterface : public FPacketHandlerComponentModuleInterface
 {
@@ -65,8 +69,4 @@ public:
 	virtual HandlerComponent* CreateComponentInstance() override;
 	virtual void StartupModule() override;
 	virtual void ShutdownModule() override;
-};
-
-DECLARE_DWORD_COUNTER_STAT_EXTERN(TEXT("Oodle Bytes Raw"), STAT_Oodle_BytesRaw, STATGROUP_Net, );
-DECLARE_DWORD_COUNTER_STAT_EXTERN(TEXT("Oodle Bytes Compressed"), STAT_Oodle_BytesCompressed, STATGROUP_Net, );
-#endif
\ No newline at end of file
+};
\ No newline at end of file
