Updated, cleaned up, refactored and tweaked the Oodle integration.

- Added batch files to help with QA - FirstRun.bat helps get setup in Training (packet capture) mode, and BuildDictionary.bat helps transform packet captures into a dictionary for compression net traffic.

- Training packet captures are timestamped and output to Game\\Saved\\Oodle\\Server into both an Output (server send) and Input (server receive) directory
- Now uses separate dictionaries for server and client.
- Oodle Binaries setup to automatically dynamically load from NotForLicensees

- Having '-ForceLogFlush' on the commandline, protects against incomplete packet capture files
- Fixed some potential exploits with packet serialization, by limiting max packet size

diff --git a/Engine/Plugins/Runtime/PacketHandlers/CompressionComponents/Oodle/BuildDictionary.bat b/Engine/Plugins/Runtime/PacketHandlers/CompressionComponents/Oodle/BuildDictionary.bat
new file mode 100644
--- /dev/null
+++ b/Engine/Plugins/Runtime/PacketHandlers/CompressionComponents/Oodle/BuildDictionary.bat
@@ -0,0 +1,163 @@
+@echo off
+
+
+echo This batch file, takes .ucap packet capture files from Oodle 'training' mode play sessions, and processes them into a dictionary file.
+echo.
+
+
+REM This batch file can only work from within the Oodle folder.
+set BaseFolder="..\..\..\..\..\.."
+
+if exist %BaseFolder:"=%\Engine goto SetUE4Editor
+
+echo Could not locate Engine folder. This .bat must be located within UE4\Engine\Plugins\Runtime\PacketHandlers\CompressionComponents\Oodle
+goto End
+
+
+:SetUE4Editor
+set UE4EditorLoc="%BaseFolder:"=%\Engine\Binaries\Win64\UE4Editor.exe"
+
+if exist %UE4EditorLoc:"=% goto SetTrainer
+
+echo Could not locate UE4Editor.exe
+goto End
+
+
+:SetTrainer
+set TrainerLoc="Source\ThirdParty\NotForLicensees\OodleTrainer\x64\Release\OodleTrainer.exe"
+
+if exist %TrainerLoc:"=% goto GetGame
+
+echo Could not locate OodleTrainer.exe, make sure it is built.
+goto End
+
+
+:GetGame
+set /p GameName=Type the name of the game you are working with: 
+echo.
+
+if exist %BaseFolder:"=%\%GameName:"=% goto GetCaptures
+if exist %BaseFolder:"=%\Samples\Games\%GameName:"=% goto GetCaptures
+
+set /p SkipGameCheck=Could not locate folder for game '%GameName:"=%'; continue anyway y/n?
+if not "%SkipGameCheck%"=="y" goto End
+
+
+:GetCaptures
+echo Place all .ucap files you want to process, into a folder, and then drag+drop the folder over this prompt.
+echo Output and Input capture files, should be processed separately.
+
+set /p CapturesFolder=Drag or paste captures folder here: 
+echo.
+
+if exist %CapturesFolder% goto GetOutput
+
+echo Could not locate folder: %CapturesFolder%
+goto End
+
+
+:GetOutput
+set BaseOutputPath="%BaseFolder:"=%\%GameName:"=%\Content\Oodle"
+
+echo Type the desired name of the dictionary file. This will be placed in %GameName:"=%\Content\Oodle\
+set /p OutputFile=Dictionary file name: 
+echo.
+
+if not exist %BaseOutputPath:"=%\%OutputFile:"=% goto ExecuteCommandlet
+
+set /p OverwriteOutput=Output file already exists. Overwrite y/n?
+if not "%OverwriteOutput%"=="y" goto End
+
+echo.
+
+
+:ExecuteCommandlet
+set MergedPacketsFile="%~dp0TempMerge.ucap"
+set CommandletParms=-run=OodleHandlerComponent.OodleTrainerCommandlet MergePackets "%MergedPacketsFile:"=%" all "%CapturesFolder:"=%
+set FinalUE4CmdLine=%GameName:"=% %CommandletParms%
+
+echo Executing Oodle MergePackets commandlet - commandline:
+echo %FinalUE4CmdLine%
+
+@echo on
+%UE4EditorLoc:"=% %FinalUE4CmdLine%
+@echo off
+echo.
+
+
+if %errorlevel%==0 goto ExecuteTrainer
+
+echo WARNING! Detected error! Review logfile for commandlet.
+goto CommandletCleanup
+
+
+:ExecuteTrainer
+if not exist %BaseFolder:"=%\%GameName:"=% goto AltExecuteTrainer
+
+set FinalOutputFile="%~dp0%BaseOutputPath:"=%\%OutputFile:"=%"
+set FinalTrainerCmdLine="%MergedPacketsFile:"=%" "%FinalOutputFile:"=%"
+
+echo Executing trainer for producing final dictionary file - commandline:
+echo %FinalTrainerCmdLine%
+
+@echo on
+%TrainerLoc:"=% %FinalTrainerCmdLine%
+@echo off
+
+echo.
+
+if not %errorlevel%==0 goto TrainerError
+goto TrainerSuccess
+
+
+:AltExecuteTrainer
+set AltOutputFile="%~dp0%OutputFile:"=%"
+set FinalAltTrainerCmdLine="%MergedPacketsFile:"=%" %AltOutputFile%
+
+echo Executing trainer for producing final dictionary file - commandline:
+echo %FinalAltTrainerCmdLine%
+
+@echo on
+%TrainerLoc:"=% %FinalAltTrainerCmdLine%
+@echo off
+
+if not %errorlevel%==0 goto TrainerError
+
+
+echo.
+echo As the game folder was not detected, you need to manually copy:
+echo %AltOutputFile%
+echo.
+echo To the folder:
+echo %GameName:"=%\Content\Oodle
+echo.
+pause
+
+
+:TrainerSuccess
+echo.
+echo Dictionary successfully generated! Set '*Dictionary=Content\Oodle\%OutputFile:"=%' under [OodleHandlerComponent], in *Engine.ini.
+echo If you want to switch Oodle from 'Training' mode to 'Release' mode, set 'Mode=Release' in the same location.
+goto PostTrainer
+echo.
+
+
+:TrainerError
+echo Error generating dictionary file. Review trainer output.
+goto End
+
+
+:PostTrainer
+:CommandletCleanup
+del %MergedPacketsFile% 
+
+
+
+:End
+echo Execution complete.
+pause
+
+
+REM Put nothing past here.
+
+
diff --git a/Engine/Plugins/Runtime/PacketHandlers/CompressionComponents/Oodle/FirstRun.bat b/Engine/Plugins/Runtime/PacketHandlers/CompressionComponents/Oodle/FirstRun.bat
new file mode 100644
--- /dev/null
+++ b/Engine/Plugins/Runtime/PacketHandlers/CompressionComponents/Oodle/FirstRun.bat
@@ -0,0 +1,75 @@
+@echo off
+
+echo This batch file, enables the Oodle plugin, and enables it as a packet handler.
+echo.
+
+
+REM This batch file can only work from within the Oodle folder.
+set BaseFolder="..\..\..\..\..\.."
+
+if exist %BaseFolder:"=%\Engine goto SetUE4Editor
+
+echo Could not locate Engine folder. This .bat must be located within UE4\Engine\Plugins\Runtime\PacketHandlers\CompressionComponents\Oodle
+goto End
+
+
+:SetUE4Editor
+set UE4EditorLoc="%BaseFolder:"=%\Engine\Binaries\Win64\UE4Editor.exe"
+
+if exist %UE4EditorLoc:"=% goto GetGame
+
+echo Could not locate UE4Editor.exe
+goto End
+
+
+:GetGame
+set /p GameName=Type the name of the game you are working with: 
+echo.
+
+
+
+:EnablePlugin
+set EnableCommandletParms=-run=plugin enable Oodle
+set FinalEnableCmdLine=%GameName:"=% %EnableCommandletParms%
+
+echo Executing plugin enable commandlet - commandline:
+echo %FinalEnableCmdLine%
+
+@echo on
+%UE4EditorLoc:"=% %FinalEnableCmdLine%
+@echo off
+echo.
+
+
+if %errorlevel%==0 goto EnableHandler
+
+echo WARNING! Detected error, plugin may not have been enabled. Will attempt to run Oodle enable commandlet.
+pause
+
+
+:EnableHandler
+set HandlerCommandletParms=-run=OodleHandlerComponent.OodleTrainerCommandlet enable
+set FinalHandlerCmdLine=%GameName:"=% %HandlerCommandletParms%
+
+
+echo Executing Oodle PacketHandler enable commandlet - commandline:
+echo %FinalHandlerCmdLine%
+
+@echo on
+%UE4EditorLoc:"=% %FinalHandlerCmdLine%
+@echo off
+echo.
+
+
+if %errorlevel%==0 goto End
+
+echo WARNING! Detected error when executing PacketHandler enable commandlet. Review the logfile.
+
+
+:End
+echo Execution complete.
+pause
+
+
+REM Put nothing past here.
+
diff --git a/Engine/Plugins/Runtime/PacketHandlers/CompressionComponents/Oodle/Oodle.uplugin b/Engine/Plugins/Runtime/PacketHandlers/CompressionComponents/Oodle/Oodle.uplugin
--- a/Engine/Plugins/Runtime/PacketHandlers/CompressionComponents/Oodle/Oodle.uplugin
+++ b/Engine/Plugins/Runtime/PacketHandlers/CompressionComponents/Oodle/Oodle.uplugin
@@ -1,7 +1,7 @@
 {
 	"FileVersion" : 3,
 	
-	"FriendlyName" : "Oddle network compression",
+	"FriendlyName" : "Oodle network compression",
 	"Version" : 1,
 	"VersionName" : "1.0",
 	"CreatedBy" : "Epic Games, Inc.",
diff --git a/Engine/Plugins/Runtime/PacketHandlers/CompressionComponents/Oodle/Source/OodleHandlerComponent/Classes/OodleTrainerCommandlet.h b/Engine/Plugins/Runtime/PacketHandlers/CompressionComponents/Oodle/Source/OodleHandlerComponent/Classes/OodleTrainerCommandlet.h
new file mode 100644
--- /dev/null
+++ b/Engine/Plugins/Runtime/PacketHandlers/CompressionComponents/Oodle/Source/OodleHandlerComponent/Classes/OodleTrainerCommandlet.h
@@ -0,0 +1,45 @@
+// Copyright 1998-2015 Epic Games, Inc. All Rights Reserved.
+
+#pragma once
+
+#include "Commandlets/Commandlet.h"
+
+#include "OodleTrainerCommandlet.generated.h"
+
+
+/**
+ * Commandlet for processing UE4 packet captures, through Oodle's training API, for generating compressed state dictionaries.
+ *
+ * Commands:
+ *	- "Enable":
+ *		- Inserts the Oodle PacketHandler into the games packet handler component list, and initializes Oodle *Engine.ini settings
+ *
+ *
+ *	- "MergePackets OutputFile PacketFile1,PacketFile2,PacketFileN":
+ *		- Takes the specified packet capture files, and merges them into a single packet capture file
+ *
+ *	- "MergePackets OutputFile All Directory":
+ *		- As above, but merges all capture files in the specified directory.
+ *
+ *
+ * @todo #JohnB: Unimplemented commands:
+ *	- "GenerateDictionary OutputFile PacketFile"
+ *		- Takes the specified packet capture file, and processes it into an dictionary state file, used for compressing packet data.
+ *		- The packet file (usually consisting of many merged packet files), should be at least 100MB in size (see Oodle docs)
+ *		- @todo #JohnB: Implement
+ *
+ *
+ *	- "PacketInfo PacketFile":
+ *		- Outputs information about the packet file, such as the MB amount of data recorded, per net connection channel, and data types
+ *		- @todo #JohnB: Only implement, if deciding to actually capture/track this kind of data
+ */
+UCLASS()
+class UOodleTrainerCommandlet : public UCommandlet
+{
+	GENERATED_UCLASS_BODY()
+
+public:
+	virtual int32 Main(const FString& Params) override;
+};
+
+
diff --git a/Engine/Plugins/Runtime/PacketHandlers/CompressionComponents/Oodle/Source/OodleHandlerComponent/OodleHandlerComponent.Build.cs b/Engine/Plugins/Runtime/PacketHandlers/CompressionComponents/Oodle/Source/OodleHandlerComponent/OodleHandlerComponent.Build.cs
--- a/Engine/Plugins/Runtime/PacketHandlers/CompressionComponents/Oodle/Source/OodleHandlerComponent/OodleHandlerComponent.Build.cs
+++ b/Engine/Plugins/Runtime/PacketHandlers/CompressionComponents/Oodle/Source/OodleHandlerComponent/OodleHandlerComponent.Build.cs
@@ -9,13 +9,26 @@ public class OodleHandlerComponent : ModuleRules
     {
 		BinariesSubFolder = "NotForLicensees";
 		
-		PublicDependencyModuleNames.AddRange( new string[] { "PacketHandler", "Core" } );
+		PrivateIncludePaths.Add("OodleHandlerComponent/Private");
+
+		PublicDependencyModuleNames.AddRange(
+			new string[]
+			{
+				"PacketHandler",
+				"Core",
+				"CoreUObject",
+				"Engine"
+			});
+
 
 		bool bHaveOodleSDK = false;
 
 		if ( ( Target.Platform == UnrealTargetPlatform.Win64 ) || ( Target.Platform == UnrealTargetPlatform.Win32 ) )
         {
-			string OodleNotForLicenseesLibDir = System.IO.Path.Combine( UEBuildConfiguration.UEThirdPartySourceDirectory, "..", "..", "Plugins", "Runtime", "PacketHandlers", "CompressionComponents", "Oodle", "Source", "ThirdParty", "NotForLicensees", "Oodle", "win", "lib" );   // Check the NotForLicensees folder first
+			// Check the NotForLicensees folder first
+			string OodleNotForLicenseesLibDir = System.IO.Path.Combine( UEBuildConfiguration.UEThirdPartySourceDirectory, "..", "..",
+				"Plugins", "Runtime", "PacketHandlers", "CompressionComponents", "Oodle", "Source", "ThirdParty", "NotForLicensees",
+				"Oodle", "win", "lib" );
 
 			try
 			{
diff --git a/Engine/Plugins/Runtime/PacketHandlers/CompressionComponents/Oodle/Source/OodleHandlerComponent/Private/OodleArchives.cpp b/Engine/Plugins/Runtime/PacketHandlers/CompressionComponents/Oodle/Source/OodleHandlerComponent/Private/OodleArchives.cpp
new file mode 100644
--- /dev/null
+++ b/Engine/Plugins/Runtime/PacketHandlers/CompressionComponents/Oodle/Source/OodleHandlerComponent/Private/OodleArchives.cpp
@@ -0,0 +1,110 @@
+// Copyright 1998-2015 Epic Games, Inc. All Rights Reserved.
+
+#include "OodleHandlerComponentPCH.h"
+
+#if HAS_OODLE_SDK
+#include "OodleArchives.h"
+
+
+/**
+ * FPacketCaptureArchive
+ */
+
+void FPacketCaptureArchive::SerializeCaptureHeader()
+{
+	uint32 Magic = CAPTURE_HEADER_MAGIC;
+
+	InnerArchive << Magic;
+	InnerArchive << CaptureVersion;
+
+	if (IsLoading())
+	{
+		check(Magic == CAPTURE_HEADER_MAGIC);
+		check(CaptureVersion <= CAPTURE_FILE_VERSION);
+	}
+	else if (IsSaving() && bImmediateFlush)
+	{
+		Flush();
+	}
+}
+
+void FPacketCaptureArchive::SerializePacket(void* PacketData, uint32& PacketSize)
+{
+	uint32 PacketBufferSize = (IsLoading() ? PacketSize : 0);
+	InnerArchive << PacketSize;
+
+	if (IsLoading())
+	{
+		// Max 128MB packet - excessive, but this is not meant to be a perfect security check
+		if (PacketBufferSize < PacketSize || !ensure(PacketSize <= 134217728))
+		{
+			UE_LOG(OodleHandlerComponentLog, Warning, TEXT("Bad PacketSize value '%i' in loading packet capture file"), PacketSize);
+			SetError();
+
+			return;
+		}
+
+		if (PacketSize > (InnerArchive.TotalSize() - InnerArchive.Tell()))
+		{
+			UE_LOG(OodleHandlerComponentLog, Warning, TEXT("PacketSize '%i' greater than remaining file data '%i'. Truncated file? ")
+					TEXT("(run server with -forcelogflush to reduce chance of truncated capture files)"),
+					PacketSize, (InnerArchive.TotalSize() - InnerArchive.Tell()));
+
+			SetError();
+
+			return;
+		}
+	}
+
+	InnerArchive.Serialize(PacketData, PacketSize);
+
+	if (IsSaving() && bImmediateFlush)
+	{
+		Flush();
+	}
+}
+
+void FPacketCaptureArchive::AppendPacketFile(FPacketCaptureArchive& InPacketFile)
+{
+	check(IsSaving());
+	check(Tell() != 0);	// Can't append a packet before writing the header
+	check(InPacketFile.IsLoading());
+	check(InPacketFile.Tell() == 0);
+
+	// Read past the header
+	InPacketFile.SerializeCaptureHeader();
+
+	check(CaptureVersion == InPacketFile.CaptureVersion);
+
+	// For appending, only support 1MB packets
+	const uint32 BufferSize = 1048576;
+	void* ReadBuffer = new uint8[BufferSize];
+
+	// Iterate through all packets
+	while (InPacketFile.Tell() < InPacketFile.TotalSize())
+	{
+		uint32 PacketSize = BufferSize;
+
+		InPacketFile.SerializePacket(ReadBuffer, PacketSize);
+
+		if (InPacketFile.IsError())
+		{
+			UE_LOG(OodleHandlerComponentLog, Warning, TEXT("Error reading packet capture data. Skipping rest of file."));
+
+			break;
+		}
+
+
+		SerializePacket(ReadBuffer, PacketSize);
+	}
+
+	delete[] ReadBuffer;
+
+
+	if (IsSaving() && bImmediateFlush)
+	{
+		Flush();
+	}
+}
+#endif
+
diff --git a/Engine/Plugins/Runtime/PacketHandlers/CompressionComponents/Oodle/Source/OodleHandlerComponent/Private/OodleHandlerComponent.cpp b/Engine/Plugins/Runtime/PacketHandlers/CompressionComponents/Oodle/Source/OodleHandlerComponent/Private/OodleHandlerComponent.cpp
--- a/Engine/Plugins/Runtime/PacketHandlers/CompressionComponents/Oodle/Source/OodleHandlerComponent/Private/OodleHandlerComponent.cpp
+++ b/Engine/Plugins/Runtime/PacketHandlers/CompressionComponents/Oodle/Source/OodleHandlerComponent/Private/OodleHandlerComponent.cpp
@@ -1,15 +1,32 @@
 // Copyright 1998-2015 Epic Games, Inc. All Rights Reserved.
 
-#include "OodleHandlerComponent.h"
+#include "OodleHandlerComponentPCH.h"
 
-#if HAS_OODLE_SDK == 1
-IMPLEMENT_MODULE( FOodleComponentModuleInterface, OodleHandlerComponent );
 
 DEFINE_LOG_CATEGORY(OodleHandlerComponentLog);
 
+
+#if HAS_OODLE_SDK
+#define OODLE_INI_SECTION TEXT("OodleHandlerComponent")
+
+
+// @todo #JohnB: Find a better solution than this, for handling the maximum packet size. This is far higher than it needs to be.
+#define MAX_OODLE_PACKET_SIZE 65535
+
+
 DEFINE_STAT(STAT_Oodle_BytesRaw);
 DEFINE_STAT(STAT_Oodle_BytesCompressed);
 
+
+FString GOodleSaveDir = TEXT("");
+
+
+// @todo #JohnB: You need to integrate the dictionary generation into a commandlet, to support FArchive-format dictionaries
+//					(alternative - and do it this way first! - adjust OodleTrainer to read/write the FArchive version files)
+
+// @todo #JohnB: You need context-sensitive data capturing; e.g. you do NOT want to be capturing voice channel data
+
+
 // Saved compressor model file :
 // must match example_packet - or whatever you use to write the state file
 struct OodleNetwork1_SavedStates_Header
@@ -23,117 +40,325 @@ struct OodleNetwork1_SavedStates_Header
 	U32 dic_complen;
 	U32 statecompacted_size;
 	U32 statecompacted_complen;
+
+	OodleNetwork1_SavedStates_Header()
+		: magic(0)
+		, compressor(0)
+		, ht_bits(0)
+		, dic_size(0)
+		, oodle_header_version(0)
+		, dic_complen(0)
+		, statecompacted_size(0)
+		, statecompacted_complen(0)
+	{
+	}
 };
 
 // OODLE
 OodleHandlerComponent::OodleHandlerComponent()
-: BytesSaved(0)
-, LogFile(nullptr)
-, SharedDictionaryData(nullptr)
-, CompressorState(nullptr)
-, Dictionary(nullptr)
+	: BytesSaved(0)
+#if EXAMPLE_CAPTURE_FORMAT
+	, PacketLogFile(NULL)
+	, bPacketLogAppend(false)
+#else
+	, InPacketLog(NULL)
+	, OutPacketLog(NULL)
+#endif
+	, Mode(EOodleHandlerMode::Release)
+	, ServerDictionary()
+	, ClientDictionary()
+#if 0
+	, Dictionary(NULL)
+	, SharedDictionaryData(NULL)
+	, CompressorState(NULL)
+#endif
 {
 	SetActive(true);
 }
 
 OodleHandlerComponent::~OodleHandlerComponent()
 {
-	free(Dictionary);
-	free(SharedDictionaryData);
-	free(CompressorState);
+#if EXAMPLE_CAPTURE_FORMAT
+	// @todo #JohnB: PacketLogfile is not deleted. Doesn't matter though, as this is being deprecated.
+#else
+	if (OutPacketLog != NULL)
+	{
+		OutPacketLog->Close();
+		OutPacketLog->DeleteInnerArchive();
+
+		delete OutPacketLog;
+
+		OutPacketLog = NULL;
+	}
+
+	if (InPacketLog != NULL)
+	{
+		InPacketLog->Close();
+		InPacketLog->DeleteInnerArchive();
+
+		delete InPacketLog;
+
+		InPacketLog = NULL;
+	}
+#endif
+
+	if (ServerDictionary.Dictionary != NULL)
+	{
+		free(ServerDictionary.Dictionary);
+	}
+
+	if (ServerDictionary.SharedDictionaryData != NULL)
+	{
+		free(ServerDictionary.SharedDictionaryData);
+	}
+
+	if (ServerDictionary.CompressorState != NULL)
+	{
+		free(ServerDictionary.CompressorState);
+	}
+
+
+	if (ClientDictionary.Dictionary != NULL)
+	{
+		free(ClientDictionary.Dictionary);
+	}
+
+	if (ClientDictionary.SharedDictionaryData != NULL)
+	{
+		free(ClientDictionary.SharedDictionaryData);
+	}
+
+	if (ClientDictionary.CompressorState != NULL)
+	{
+		free(ClientDictionary.CompressorState);
+	}
 
 	Oodle_Shutdown_NoThreads();
 }
 
+void OodleHandlerComponent::InitFirstRunConfig()
+{
+	// Check that the OodleHandlerComponent section exists, and if not, init with defaults
+	if (!GConfig->DoesSectionExist(OODLE_INI_SECTION, GEngineIni))
+	{
+		GConfig->SetString(OODLE_INI_SECTION, TEXT("Mode"), TEXT("Training"), GEngineIni);
+		GConfig->SetString(OODLE_INI_SECTION, TEXT("PacketLogFile"), TEXT("PacketDump"), GEngineIni);
+
+#if EXAMPLE_CAPTURE_FORMAT
+		GConfig->SetBool(OODLE_INI_SECTION, TEXT("bPacketLogAppend"), false, GEngineIni);
+#endif
+
+		GConfig->SetString(OODLE_INI_SECTION, TEXT("ServerDictionary"), TEXT(""), GEngineIni);
+		GConfig->SetString(OODLE_INI_SECTION, TEXT("ClientDictionary"), TEXT(""), GEngineIni);
+	}
+}
+
 void OodleHandlerComponent::Initialize()
 {
+	InitFirstRunConfig();
+
+#if EXAMPLE_CAPTURE_FORMAT
+	// Class config variables
+	GConfig->GetBool(OODLE_INI_SECTION, TEXT("bPacketLogAppend"), bPacketLogAppend, GEngineIni);
+#endif
+
 	// Mode
 	FString ReadMode;
-	GConfig->GetString(TEXT("OodleHandlerComponent"), TEXT("Mode"), ReadMode, GEngineIni);
+	GConfig->GetString(OODLE_INI_SECTION, TEXT("Mode"), ReadMode, GEngineIni);
 
 	if(ReadMode == TEXT("Training"))
 	{
-		Mode = Training;
+		Mode = EOodleHandlerMode::Training;
 	}
 	else if(ReadMode == TEXT("Release"))
 	{
-		Mode = Release;
+		Mode = EOodleHandlerMode::Release;
 	}
 	// Default
 	else
 	{
-		Mode = Release;
+		Mode = EOodleHandlerMode::Release;
 	}
 
 	switch(Mode)
 	{
-		case Training:
+	case EOodleHandlerMode::Training:
 		{
-			if(Handler->Mode == Handler::Mode::Server)
+			if (Handler->Mode == Handler::Mode::Server)
 			{
-				FString ReadOutputLogDirectory;
-				GConfig->GetString(TEXT("OodleHandlerComponent"), TEXT("LogFile"), ReadOutputLogDirectory, GEngineIni);
+				IPlatformFile& PlatformFile = FPlatformFileManager::Get().GetPlatformFile();
+				FString ReadOutputLogDirectory = FPaths::Combine(*GOodleSaveDir, TEXT("Server"));
+				FString BaseFilename;
+
+				PlatformFile.CreateDirectoryTree(*ReadOutputLogDirectory);
+				PlatformFile.CreateDirectoryTree(*FPaths::Combine(*ReadOutputLogDirectory, TEXT("Input")));
+				PlatformFile.CreateDirectoryTree(*FPaths::Combine(*ReadOutputLogDirectory, TEXT("Output")));
+				GConfig->GetString(OODLE_INI_SECTION, TEXT("PacketLogFile"), BaseFilename, GEngineIni);
+
+				BaseFilename = FPaths::GetBaseFilename(BaseFilename);
+
+
+#if EXAMPLE_CAPTURE_FORMAT
+				FString FinalFilePath;
+
+				// When bPacketLogAppend is used, a static log filename is used
+				if (bPacketLogAppend)
+				{
+					FString PreExtFilePath = FPaths::Combine(*ReadOutputLogDirectory, *BaseFilename);
+
+					// Try0: PacketDump.ucap, Try1: PacketDump_1.ucap, Try#:PacketDump_#.ucap
+					FinalFilePath = PreExtFilePath + CAPTURE_EXT;
+
+					for (int32 i=1; ; i++)
+					{
+						// If not appending to a file, make sure to disable this bool - otherwise the file header won't be written
+						bPacketLogAppend = PlatformFile.FileExists(*FinalFilePath);
+						PacketLogFile = PlatformFile.OpenWrite(*FinalFilePath, true);
+
+						if (PacketLogFile != NULL)
+						{
+							break;
+						}
+
+
+						FinalFilePath = PreExtFilePath + FString::Printf(TEXT("_%i"), i) + CAPTURE_EXT;
+					}
+				}
+				// When non-append logging is used, the log filename uses a timestamp
+				else
+				{
+					BaseFilename = BaseFilename + TEXT("_") + FDateTime::Now().ToString();
+
+					FString PreExtFilePath = FPaths::Combine(*ReadOutputLogDirectory, *BaseFilename);
+					FinalFilePath = PreExtFilePath + CAPTURE_EXT;
+
+					for (int32 i=1; PlatformFile.FileExists(*FinalFilePath); i++)
+					{
+						FinalFilePath = PreExtFilePath + FString::Printf(TEXT("_%i"), i) + CAPTURE_EXT;
+					}
+
+					// Try0: PacketDump_2015.09.23-20.12.54.ucap, Try1: PacketDump_2015.09.23-20.12.54_1.ucap, ...
+					PacketLogFile = PlatformFile.OpenWrite(*FinalFilePath);
+				}
+
+				if (PacketLogFile == NULL)
+				{
+					LowLevelFatalError(TEXT("Failed to create file %s"), *FinalFilePath);
+					return;
+				}
+#else
+				BaseFilename = BaseFilename + TEXT("_") + FDateTime::Now().ToString();
+
+				FString PreExtInFilePath = FPaths::Combine(*ReadOutputLogDirectory, TEXT("Input"), *(BaseFilename + TEXT("_Input")));
+				FString PreExtOutFilePath = FPaths::Combine(*ReadOutputLogDirectory, TEXT("Output"), *(BaseFilename + TEXT("_Output")));
+				FString InPath = PreExtInFilePath + CAPTURE_EXT;
+				FString OutPath = PreExtOutFilePath + CAPTURE_EXT;
+
+				// Ensure the In/Out filenames are unique
+				for (int32 i=1; PlatformFile.FileExists(*InPath) || PlatformFile.FileExists(*OutPath); i++)
+				{
+					InPath = PreExtInFilePath + FString::Printf(TEXT("_%i"), i) + CAPTURE_EXT;
+					OutPath = PreExtOutFilePath + FString::Printf(TEXT("_%i"), i) + CAPTURE_EXT;
+				}
+
+				FArchive* InArc = IFileManager::Get().CreateFileWriter(*InPath);
+				FArchive* OutArc = (InArc != NULL ? IFileManager::Get().CreateFileWriter(*OutPath) : NULL);
+
+				InPacketLog = (InArc != NULL ? new FPacketCaptureArchive(*InArc) : NULL);
+				OutPacketLog = (OutArc != NULL ? new FPacketCaptureArchive(*OutArc) : NULL);
 
-				LogFile = FPlatformFileManager::Get().GetPlatformFile().OpenWrite(*ReadOutputLogDirectory);
 
-				if(!LogFile)
+				if (InPacketLog == NULL || OutPacketLog == NULL)
 				{
-					LowLevelFatalError(TEXT("Failed to create file %s"), *ReadOutputLogDirectory);
+					LowLevelFatalError(TEXT("Failed to create files '%s' and '%s'"), *InPath, *OutPath);
+					return;
 				}
+#endif
 
-				int32 NumberOfChannels = 1;
+
+#if EXAMPLE_CAPTURE_FORMAT
+				// Only write the initial header, if not appending
+				if (!bPacketLogAppend)
+				{
+					int32 NumberOfChannels = 1;
 			
-				LogFile->Write(reinterpret_cast<const uint8*>(&NumberOfChannels), sizeof(NumberOfChannels));
+					PacketLogFile->Write(reinterpret_cast<const uint8*>(&NumberOfChannels), sizeof(NumberOfChannels));
+				}
+#else
+				InPacketLog->SerializeCaptureHeader();
+				OutPacketLog->SerializeCaptureHeader();
+#endif
 			}
+
 			break;
 		}
-		case Release:
+
+	case EOodleHandlerMode::Release:
 		{
-			FString ReadStateDirectory;
-			GConfig->GetString(TEXT("OodleHandlerComponent"), TEXT("StateFile"), ReadStateDirectory, GEngineIni);
+			FString ServerDictionaryPath;
+			FString ClientDictionaryPath;
+
+			GConfig->GetString(OODLE_INI_SECTION, TEXT("ServerDictionary"), ServerDictionaryPath, GEngineIni);
+			GConfig->GetString(OODLE_INI_SECTION, TEXT("ClientDictionary"), ClientDictionaryPath, GEngineIni);
 
-			if(ReadStateDirectory.Len() < 0)
+			if (ServerDictionaryPath.Len() < 0 || ClientDictionaryPath.Len() < 0)
 			{
-				LowLevelFatalError(TEXT("Specify a state directory for Oodle compressor in DefaultEngine.ini"));
+				LowLevelFatalError(TEXT("Specify both Server/Client dictionaries for Oodle compressor in DefaultEngine.ini"));
 				return;
 			}
 
-			TArray<uint8> FileInMemory;
-			if (!FFileHelper::LoadFileToArray(FileInMemory, *ReadStateDirectory))
-			{
-				LowLevelFatalError(TEXT("Incorrect LogFile Provided"));
-			}
-	
-			// parse header :
-			OodleNetwork1_SavedStates_Header* SavedStateFileData = (OodleNetwork1_SavedStates_Header *)FileInMemory.GetData();
 
-			if (SavedStateFileData->magic != ON1_MAGIC)
-			{
-				UE_LOG(OodleHandlerComponentLog, Warning, TEXT("OodleNetCompressor: state_file ON1_MAGIC mismatch"));
-				return;
-			}
+			// Path must be within game directory, e.g: "Content/Oodle/Output.udic" becomes "ShooterGam/Content/Oodle/Output.udic"
+			ServerDictionaryPath = FPaths::Combine(*FPaths::GameDir(), *ServerDictionaryPath);
+			ClientDictionaryPath = FPaths::Combine(*FPaths::GameDir(), *ClientDictionaryPath);
 
-			if (SavedStateFileData->oodle_header_version != OODLE_HEADER_VERSION)
+			FPaths::CollapseRelativeDirectories(ServerDictionaryPath);
+			FPaths::CollapseRelativeDirectories(ClientDictionaryPath);
+
+			// Don't allow directory traversal to escape the game directory
+			if (!ServerDictionaryPath.StartsWith(FPaths::GameDir()) || !ClientDictionaryPath.StartsWith(FPaths::GameDir()))
 			{
-				UE_LOG(OodleHandlerComponentLog, Warning, TEXT("OodleNetCompressor: state_file OODLE_HEADER_VERSION mismatch. Was: %u, Expected: %u"), (uint32)SavedStateFileData->oodle_header_version, (uint32)OODLE_HEADER_VERSION);
+				LowLevelFatalError(TEXT("DictionaryFile not allowed to use ../ paths to escape game directory."));
 				return;
 			}
 
-			// this code only loads compressed states currently :
-			check(SavedStateFileData->compressor != -1);
 
-			// Init
-			OodleInitOptions Options;
-			//Oodle_Init_GetDefaults(OODLE_HEADER_VERSION, &Options);
-			Oodle_Init_GetDefaults_Minimal(OODLE_HEADER_VERSION, &Options);
-			Oodle_Init_NoThreads(OODLE_HEADER_VERSION, &Options);
+			InitializeDictionary(ServerDictionaryPath, ServerDictionary);
+			InitializeDictionary(ClientDictionaryPath, ClientDictionary);
+
+
+			break;
+		}
+
+	default:
+		break;
+	}
+
+	Initialized();
+}
+
+void OodleHandlerComponent::InitializeDictionary(FString FilePath, FOodleDictionary& OutDictionary)
+{
+	TArray<uint8> FileInMemory;
+
+	// @todo #JohnB: Need a way to verify the state file, is actually a state file (do this once converting to FArchive)
+	if (FFileHelper::LoadFileToArray(FileInMemory, *FilePath))
+	{
+		// parse header :
+		OodleNetwork1_SavedStates_Header* SavedDictionaryFileData = (OodleNetwork1_SavedStates_Header*)FileInMemory.GetData();
+
+		if (SavedDictionaryFileData->magic == ON1_MAGIC && SavedDictionaryFileData->oodle_header_version == OODLE_HEADER_VERSION)
+		{
+			// this code only loads compressed states currently
+			check(SavedDictionaryFileData->compressor != -1);
+
 
-			S32 StateHTBits = SavedStateFileData->ht_bits;
-			SINTa DictionarySize = SavedStateFileData->dic_size;
+			S32 StateHTBits = SavedDictionaryFileData->ht_bits;
+			SINTa DictionarySize = SavedDictionaryFileData->dic_size;
 
-			SINTa DictionaryCompLen = SavedStateFileData->dic_complen;
-			SINTa CompressorStateCompactedSize = SavedStateFileData->statecompacted_size;
-			SINTa CompressorStateCompactedLength = SavedStateFileData->statecompacted_complen;
+			SINTa DictionaryCompLen = SavedDictionaryFileData->dic_complen;
+			SINTa CompressorStateCompactedSize = SavedDictionaryFileData->statecompacted_size;
+			SINTa CompressorStateCompactedLength = SavedDictionaryFileData->statecompacted_complen;
 
 			SINTa SharedSize = OodleNetwork1_Shared_Size(StateHTBits);
 			SINTa CompressorStateSize = OodleNetwork1UDP_State_Size();
@@ -143,18 +368,21 @@ void OodleHandlerComponent::Initialize()
 
 			check(CompressorStateCompactedSize > 0 && CompressorStateCompactedSize < OodleNetwork1UDP_StateCompacted_MaxSize());
 
-			check(FileInMemory.Num() == (S64)sizeof(OodleNetwork1_SavedStates_Header)+DictionaryCompLen + CompressorStateCompactedLength);
+			check(FileInMemory.Num() == (S64)sizeof(OodleNetwork1_SavedStates_Header) + DictionaryCompLen +
+					CompressorStateCompactedLength);
 
-			//OodleLog_Printf_v1("OodleNetwork1UDP Loading; dic comp %d , state %d->%d\n",DictionaryCompLen, CompressorStateCompactedSize, CompressorStateCompactedLength);
+			//OodleLog_Printf_v1("OodleNetwork1UDP Loading; dic comp %d , state %d->%d\n",DictionaryCompLen,
+			//						CompressorStateCompactedSize, CompressorStateCompactedLength);
 
 			//-------------------------------------------
 			// decompress Dictionary and CompressorStatecompacted
 
-			Dictionary = malloc(DictionarySize);
+			OutDictionary.Dictionary = malloc(DictionarySize);
 
-			void * DictionaryCompressed = SavedStateFileData + 1;
+			void * DictionaryCompressed = SavedDictionaryFileData + 1;
 
-			SINTa DecompressedDictionarySize = OodleLZ_Decompress(DictionaryCompressed, DictionaryCompLen, Dictionary, DictionarySize);
+			SINTa DecompressedDictionarySize = OodleLZ_Decompress(DictionaryCompressed, DictionaryCompLen, OutDictionary.Dictionary,
+																	DictionarySize);
 
 			check(DecompressedDictionarySize == DictionarySize);
 
@@ -162,36 +390,46 @@ void OodleHandlerComponent::Initialize()
 
 			void * CompressorStatecompacted_comp_ptr = (U8 *)DictionaryCompressed + DictionaryCompLen;
 
-			SINTa decomp_statecompacted_size = OodleLZ_Decompress(CompressorStatecompacted_comp_ptr, CompressorStateCompactedLength, UDPNewCompacted, CompressorStateCompactedSize);
+			SINTa decomp_statecompacted_size = OodleLZ_Decompress(CompressorStatecompacted_comp_ptr, CompressorStateCompactedLength,
+																	UDPNewCompacted, CompressorStateCompactedSize);
 
 			check(decomp_statecompacted_size == CompressorStateCompactedSize);
 
 			//----------------------------------------------
 			// Uncompact the "Compacted" state into a usable state
 
-			CompressorState = (OodleNetwork1UDP_State *)malloc(CompressorStateSize);
-			check(CompressorState != NULL);
+			OutDictionary.CompressorState = (OodleNetwork1UDP_State *)malloc(CompressorStateSize);
+			check(OutDictionary.CompressorState != NULL);
 
-			OodleNetwork1UDP_State_Uncompact(CompressorState, UDPNewCompacted);
+			OodleNetwork1UDP_State_Uncompact(OutDictionary.CompressorState, UDPNewCompacted);
 			free(UDPNewCompacted);
 
 			//----------------------------------------------
 			// fill out SharedDictionaryData from the dictionary
 
-			SharedDictionaryData = (OodleNetwork1_Shared *)malloc(SharedSize);
-			check(SharedDictionaryData != NULL);
-			OodleNetwork1_Shared_SetWindow(SharedDictionaryData, StateHTBits, Dictionary, (S32)DictionarySize);
+			OutDictionary.SharedDictionaryData = (OodleNetwork1_Shared *)malloc(SharedSize);
+			check(OutDictionary.SharedDictionaryData != NULL);
+			OodleNetwork1_Shared_SetWindow(OutDictionary.SharedDictionaryData, StateHTBits, OutDictionary.Dictionary, (S32)DictionarySize);
 
 			//-----------------------------------------------------
 	
 			//SET_DWORD_STAT(STAT_Oodle_DicSize, DictionarySize);
-			break;
 		}
-		default:
-		break;
+		else if (SavedDictionaryFileData->magic != ON1_MAGIC)
+		{
+			UE_LOG(OodleHandlerComponentLog, Warning, TEXT("OodleNetCompressor: state_file ON1_MAGIC mismatch"));
+		}
+		else //if (SavedDictionaryFileData->oodle_header_version != OODLE_HEADER_VERSION)
+		{
+			UE_LOG(OodleHandlerComponentLog, Warning,
+					TEXT("OodleNetCompressor: state_file OODLE_HEADER_VERSION mismatch. Was: %u, Expected: %u"),
+					(uint32)SavedDictionaryFileData->oodle_header_version, (uint32)OODLE_HEADER_VERSION);
+		}
+	}
+	else
+	{
+		LowLevelFatalError(TEXT("Incorrect DictionaryFile Provided"));
 	}
-
-	Initialized();
 }
 
 bool OodleHandlerComponent::IsValid() const
@@ -203,38 +441,51 @@ void OodleHandlerComponent::Incoming(FBitReader& Packet)
 {
 	switch(Mode)
 	{
-		case Training:
+	case EOodleHandlerMode::Training:
 		{
 			if(Handler->Mode == Handler::Mode::Server)
 			{
 				uint32 SizeOfPacket = Packet.GetNumBytes();
+
+#if EXAMPLE_CAPTURE_FORMAT
 				int32 channel = 0;
 
-				LogFile->Write(reinterpret_cast<const uint8*>(&channel), sizeof(channel));
-				LogFile->Write(reinterpret_cast<const uint8*>(&SizeOfPacket), sizeof(SizeOfPacket));
-				LogFile->Write(reinterpret_cast<const uint8*>(Packet.GetData()), SizeOfPacket);
+				PacketLogFile->Write(reinterpret_cast<const uint8*>(&channel), sizeof(channel));
+				PacketLogFile->Write(reinterpret_cast<const uint8*>(&SizeOfPacket), sizeof(SizeOfPacket));
+				PacketLogFile->Write(reinterpret_cast<const uint8*>(Packet.GetData()), SizeOfPacket);
+#else
+				void* PacketData = Packet.GetData();
+				InPacketLog->SerializePacket(PacketData, SizeOfPacket);
+#endif
 			}
 			break;
 		}
-		case Release:
+
+	case EOodleHandlerMode::Release:
 		{
-			// Remove size
 			uint32 DecompressedLength;
 			Packet.SerializeIntPacked(DecompressedLength);
 
-			const int32 CompressedLength = Packet.GetBytesLeft();
-			uint8* CompressedData = new uint8[CompressedLength];
-			Packet.Serialize(CompressedData, CompressedLength);
+			if (ensure(DecompressedLength < MAX_OODLE_PACKET_SIZE))
+			{
+				// @todo #JohnB: You want to tweak the size of 'MAX_OODLE_PACKET_SIZE'
+				static uint8 CompressedData[MAX_OODLE_PACKET_SIZE];
+				static uint8 DecompressedData[MAX_OODLE_PACKET_SIZE];
+				const int32 CompressedLength = Packet.GetBytesLeft();
+				FOodleDictionary* CurDict = ((Handler->Mode == Handler::Mode::Server) ? &ClientDictionary : &ServerDictionary);
+
+				Packet.Serialize(CompressedData, CompressedLength);
+
+				check(OodleNetwork1UDP_Decode(CurDict->CompressorState, CurDict->SharedDictionaryData, CompressedData, CompressedLength,
+												DecompressedData, DecompressedLength));
 
-			uint8* DecompressedData = new uint8[DecompressedLength];
-	
-			check(OodleNetwork1UDP_Decode(CompressorState, SharedDictionaryData, CompressedData, CompressedLength, DecompressedData, DecompressedLength));
 
-			FBitReader UnCompressedPacket(DecompressedData, DecompressedLength * 8);
-			Packet = UnCompressedPacket;
+				// @todo #JohnB: Decompress directly into FBitReader's buffer
+				FBitReader UnCompressedPacket(DecompressedData, DecompressedLength * 8);
+				Packet = UnCompressedPacket;
+			}
+
 
-			delete[] CompressedData;
-			delete[] DecompressedData;
 			break;
 		}
 	}
@@ -244,44 +495,59 @@ void OodleHandlerComponent::Outgoing(FBitWriter& Packet)
 {
 	switch(Mode)
 	{
-		case Training:
+	case EOodleHandlerMode::Training:
 		{
 			if(Handler->Mode == Handler::Mode::Server)
 			{
 				uint32 SizeOfPacket = Packet.GetNumBytes();
+
+#if EXAMPLE_CAPTURE_FORMAT
 				int32 channel = 0;
-				LogFile->Write(reinterpret_cast<const uint8*>(&channel), sizeof(channel));
-				LogFile->Write(reinterpret_cast<const uint8*>(&SizeOfPacket), sizeof(SizeOfPacket));
-				LogFile->Write(reinterpret_cast<const uint8*>(Packet.GetData()), SizeOfPacket);
+
+				PacketLogFile->Write(reinterpret_cast<const uint8*>(&channel), sizeof(channel));
+				PacketLogFile->Write(reinterpret_cast<const uint8*>(&SizeOfPacket), sizeof(SizeOfPacket));
+				PacketLogFile->Write(reinterpret_cast<const uint8*>(Packet.GetData()), SizeOfPacket);
+#else
+				OutPacketLog->SerializePacket((void*)Packet.GetData(), SizeOfPacket);
+#endif
 			}
 			break;
 		}
-		case Release:
+
+	case EOodleHandlerMode::Release:
 		{
 			// Add size
 			uint32 UncompressedLength = Packet.GetNumBytes();
 
-			if(UncompressedLength > 0)
+			check(UncompressedLength < MAX_OODLE_PACKET_SIZE);
+
+			if (UncompressedLength > 0)
 			{
-				uint8* UncompressedData = new uint8[UncompressedLength];
+				// @todo #JohnB: You should be able to share the same two buffers between Incoming and Outgoing, as an optimization
+				static uint8 UncompressedData[MAX_OODLE_PACKET_SIZE];
+				static uint8 CompressedData[MAX_OODLE_PACKET_SIZE];
+				FOodleDictionary* CurDict = ((Handler->Mode == Handler::Mode::Server) ? &ServerDictionary : &ClientDictionary);
+
+				check(OodleLZ_GetCompressedBufferSizeNeeded(UncompressedLength) < MAX_OODLE_PACKET_SIZE);
+
 				memcpy(UncompressedData, Packet.GetData(), UncompressedLength);
 
-				uint8_t* CompressedData = new uint8_t[OodleLZ_GetCompressedBufferSizeNeeded(UncompressedLength)];
+				SINTa CompressedLengthSINT = OodleNetwork1UDP_Encode(CurDict->CompressorState, CurDict->SharedDictionaryData,
+																		UncompressedData, UncompressedLength, CompressedData);
 
-				SINTa CompressedLengthSINT = OodleNetwork1UDP_Encode(CompressorState, SharedDictionaryData, UncompressedData, UncompressedLength, CompressedData);
 				uint32 CompressedLength = static_cast<uint32>(CompressedLengthSINT);
 
 				Packet.Reset();
 
+				// @todo #JohnB: Compress directly into a (deliberately oversized) FBitWriter buffer, which you can shrink afterwards
 				Packet.SerializeIntPacked(UncompressedLength);
 				Packet.Serialize(CompressedData, CompressedLength);
 
-				delete UncompressedData;
-				delete CompressedData;
 
 				BytesSaved += UncompressedLength - CompressedLength;
 
-				//UE_LOG(OodleHandlerComponentLog, Log, TEXT("Oodle Compressed: UnCompressed: %i Compressed: %i Total: %i"), UncompressedLength, CompressedLength, BytesSaved);
+				//UE_LOG(OodleHandlerComponentLog, Log, TEXT("Oodle Compressed: UnCompressed: %i Compressed: %i Total: %i"),
+				//			UncompressedLength, CompressedLength, BytesSaved);
 
 				INC_DWORD_STAT_BY(STAT_Oodle_BytesRaw, UncompressedLength);
 				INC_DWORD_STAT_BY(STAT_Oodle_BytesCompressed, CompressedLength);
@@ -297,26 +563,73 @@ void OodleHandlerComponent::Outgoing(FBitWriter& Packet)
 }
 
 // MODULE INTERFACE
-HandlerComponent* FOodleComponentModuleInterface::CreateComponentInstance()
+TSharedPtr<HandlerComponent> FOodleComponentModuleInterface::CreateComponentInstance(FString& Options)
 {
-	return new OodleHandlerComponent;
+	return MakeShareable(new OodleHandlerComponent);
 }
 
 void FOodleComponentModuleInterface::StartupModule()
 {
-	Oodle_Init_Default(OODLE_HEADER_VERSION, Oodle_Init_GetDefaults_DebugSystems_No,Oodle_Init_GetDefaults_Threads_No);
+	// Use an absolute path for this, as we want all relative paths, to be relative to this folder
+	GOodleSaveDir = FPaths::ConvertRelativePathToFull(FPaths::Combine(*FPaths::GameSavedDir(), TEXT("Oodle")));
+
+
+	// Load the Oodle library
+	FString OodleBinaryPath = FPaths::EngineDir() / TEXT("Binaries/ThirdParty/NotForLicensees/Oodle/");
+	FString OodleBinaryFile;
+
+#if PLATFORM_WINDOWS
+	{
+	#if PLATFORM_64BITS
+		OodleBinaryPath += TEXT("Win64/");
+		OodleBinaryFile = TEXT("oodle_145_win64.dll");
+	#else
+		OodleBinaryPath += TEXT("Win32/");
+		OodleBinaryFile = TEXT("oodle_145_win32.dll");
+	#endif
+
+
+		FPlatformProcess::PushDllDirectory(*OodleBinaryPath);
+
+		OodleDllHandle = FPlatformProcess::GetDllHandle(*(OodleBinaryPath + OodleBinaryFile));
+
+		FPlatformProcess::PopDllDirectory(*OodleBinaryPath);
+	}
+#else
+	// @todo #JohnB
+#endif
+
+	check(OodleDllHandle != NULL);
+
+
+	OodleInitOptions Options;
+
+	Oodle_Init_GetDefaults_Minimal(OODLE_HEADER_VERSION, &Options);
+
+	// @todo #JohnB: Implement toggle for logfile (from config), and implement/test the log filename path (don't forget char conversion)
+	//Options.m_OodleInit_Log = true;
+	//Options.m_OodleInit_Log_Header = ?;
+	//Options.m_OodleInit_Log_FileName = ?;
+	//Options.m_OodleInit_Log_FlushEachWrite = FParse::Param(FCommandLine::Get(), TEXT("FORCELOGFLUSH"));
+
+	// @todo #JohnB: Consider threading in future
+	Oodle_Init_NoThreads(OODLE_HEADER_VERSION, &Options);
 }
 
 void FOodleComponentModuleInterface::ShutdownModule()
 {
 	Oodle_Shutdown();
+
+	if (OodleDllHandle != NULL)
+	{
+		FPlatformProcess::FreeDllHandle(OodleDllHandle);
+		OodleDllHandle = NULL;
+	}
 }
 #else
-IMPLEMENT_MODULE( FOodleComponentModuleInterface, OodleHandlerComponent );
-
-HandlerComponent* FOodleComponentModuleInterface::CreateComponentInstance()
+TSharedPtr<HandlerComponent> FOodleComponentModuleInterface::CreateComponentInstance(FString& Options)
 {
-	return nullptr;
+	return TSharedPtr<HandlerComponent>(NULL);
 }
 
 void FOodleComponentModuleInterface::StartupModule()
@@ -326,4 +639,6 @@ void FOodleComponentModuleInterface::StartupModule()
 void FOodleComponentModuleInterface::ShutdownModule()
 {
 }
-#endif
\ No newline at end of file
+#endif
+
+IMPLEMENT_MODULE( FOodleComponentModuleInterface, OodleHandlerComponent );
\ No newline at end of file
diff --git a/Engine/Plugins/Runtime/PacketHandlers/CompressionComponents/Oodle/Source/OodleHandlerComponent/Private/OodleHandlerComponentPCH.h b/Engine/Plugins/Runtime/PacketHandlers/CompressionComponents/Oodle/Source/OodleHandlerComponent/Private/OodleHandlerComponentPCH.h
new file mode 100644
--- /dev/null
+++ b/Engine/Plugins/Runtime/PacketHandlers/CompressionComponents/Oodle/Source/OodleHandlerComponent/Private/OodleHandlerComponentPCH.h
@@ -0,0 +1,6 @@
+// Copyright 1998-2015 Epic Games, Inc. All Rights Reserved.
+
+#pragma once
+
+#include "OodleHandlerComponent.h"
+
diff --git a/Engine/Plugins/Runtime/PacketHandlers/CompressionComponents/Oodle/Source/OodleHandlerComponent/Private/OodleTrainerCommandlet.cpp b/Engine/Plugins/Runtime/PacketHandlers/CompressionComponents/Oodle/Source/OodleHandlerComponent/Private/OodleTrainerCommandlet.cpp
new file mode 100644
--- /dev/null
+++ b/Engine/Plugins/Runtime/PacketHandlers/CompressionComponents/Oodle/Source/OodleHandlerComponent/Private/OodleTrainerCommandlet.cpp
@@ -0,0 +1,293 @@
+// Copyright 1998-2015 Epic Games, Inc. All Rights Reserved.
+
+#include "OodleHandlerComponentPCH.h"
+
+#include "OodleTrainerCommandlet.h"
+
+
+#define LOCTEXT_NAMESPACE "Oodle"
+
+
+// @todo #JohnB: Make sure to note that >100mb data (for production games, even more) should be captured
+
+// @todo #JohnB: Gather stats while merging, to give the commandlet some more informative output
+
+/**
+ * UOodleTrainerCommandlet
+ */
+
+UOodleTrainerCommandlet::UOodleTrainerCommandlet(const FObjectInitializer& ObjectInitializer)
+	: Super(ObjectInitializer)
+{
+}
+
+int32 UOodleTrainerCommandlet::Main(const FString& Params)
+{
+#if HAS_OODLE_SDK
+	TArray<FString> Tokens, Switches;
+
+	ParseCommandLine(*Params, Tokens, Switches);
+
+	if (Tokens.Num() > 0)
+	{
+		FString MainCmd = Tokens[0];
+
+		/**
+		 * Enable
+		 *
+		 * Command for quick-enabling Oodle
+		 */
+		if (MainCmd == TEXT("Enable"))
+		{
+			TArray<FString> Components;
+			bool bAlreadyEnabled = false;
+
+			GConfig->GetArray(TEXT("PacketHandlerComponents"), TEXT("Components"), Components, GEngineIni);
+
+			for (FString CurComponent : Components)
+			{
+				if (CurComponent.StartsWith(TEXT("OodleHandlerComponent")))
+				{
+					bAlreadyEnabled = true;
+					break;
+				}
+			}
+
+
+			if (!bAlreadyEnabled)
+			{
+				Components.Add(TEXT("OodleHandlerComponent"));
+				GConfig->SetArray(TEXT("PacketHandlerComponents"), TEXT("Components"), Components, GEngineIni);
+
+				OodleHandlerComponent::InitFirstRunConfig();
+
+
+				UE_LOG(OodleHandlerComponentLog, Log, TEXT("Initialized first-run settings for Oodle (will start in Trainer mode)."));
+			}
+			else
+			{
+				UE_LOG(OodleHandlerComponentLog, Error, TEXT("The Oodle PacketHandler is already enabled."));
+			}
+		}
+		/**
+		 * MergePackets
+		 *
+		 * Command for merging capture packets
+		 */
+		else if (MainCmd == TEXT("MergePackets"))
+		{
+			if (Tokens.Num() > 2)
+			{
+				FString OutputCapFile = Tokens[1];
+				FString MergeList = Tokens[2];
+				IPlatformFile& PlatformFile = FPlatformFileManager::Get().GetPlatformFile();
+				TArray<FString> MergeFiles;
+				bool bValidFileParms = true;
+
+
+				if (MergeList == TEXT("all"))
+				{
+					bValidFileParms = false;
+
+					// @todo #JohnB: Document
+
+					if (Tokens.Num() > 3)
+					{
+						FString MergeDir = Tokens[3];
+
+						if (FPaths::IsRelative(MergeDir))
+						{
+							MergeDir = FPaths::Combine(*GOodleSaveDir, TEXT("Server"), *MergeDir);
+						}
+
+						if (FPaths::DirectoryExists(MergeDir))
+						{
+							IFileManager::Get().FindFiles(MergeFiles, *MergeDir, TEXT("ucap"));
+
+							if (MergeFiles.Num())
+							{
+								bValidFileParms = true;
+
+								for (int32 i=0; i<MergeFiles.Num(); i++)
+								{
+									MergeFiles[i] = FPaths::Combine(*MergeDir, *MergeFiles[i]);
+								}
+							}
+							else
+							{
+								UE_LOG(OodleHandlerComponentLog, Error, TEXT("Could not find any .ucap files in directory: %s"),
+										*MergeDir);
+							}
+						}
+						else
+						{
+							UE_LOG(OodleHandlerComponentLog, Error, TEXT("Could not find directory: %s"), *MergeDir);
+						}
+					}
+					else
+					{
+						UE_LOG(OodleHandlerComponentLog, Error,
+								TEXT("Need to specify a target directory, e.g: MergePackets OutFile All Dir"));
+					}
+				}
+				else
+				{
+					MergeList.ParseIntoArray(MergeFiles, TEXT(","));
+				}
+
+
+				// Adjust non-absolute directories, so they are relative to the Saved\Oodle\Server directory
+				if (bValidFileParms && FPaths::IsRelative(OutputCapFile))
+				{
+					OutputCapFile = FPaths::Combine(*GOodleSaveDir, TEXT("Server"), *OutputCapFile);
+
+					if (FPaths::FileExists(OutputCapFile))
+					{
+						if (FApp::IsUnattended())
+						{
+							// @todo #JohnB: Error message
+
+							bValidFileParms = false;
+						}
+						else
+						{
+							EAppReturnType::Type Result = FMessageDialog::Open(EAppMsgType::YesNo,
+																		LOCTEXT("MergeOverwrite", "Overwrite existing output file?"));
+
+							bValidFileParms = (Result == EAppReturnType::Yes);
+						}
+					}
+				}
+
+				for (int32 i=0; i<MergeFiles.Num() && bValidFileParms; i++)
+				{
+					FString CurMergeFile = MergeFiles[i];
+
+					if (FPaths::IsRelative(CurMergeFile))
+					{
+						MergeFiles[i] = FPaths::Combine(*GOodleSaveDir, TEXT("Server"), *CurMergeFile);
+
+						if (!FPaths::FileExists(CurMergeFile))
+						{
+							UE_LOG(OodleHandlerComponentLog, Error, TEXT("Packet capture file does not exist: %s"), *CurMergeFile);
+
+							bValidFileParms = false;
+						}
+					}
+				}
+
+
+				if (bValidFileParms && MergeFiles.Num() > 1)
+				{
+					bool bAllMergeHandlesValid = true;
+					TMap<FArchive*, FString> MergeMap;
+
+					for (FString CurMergeFile : MergeFiles)
+					{
+						FArchive* ReadArc = IFileManager::Get().CreateFileReader(*CurMergeFile);
+
+						if (ReadArc != NULL)
+						{
+							MergeMap.Add(ReadArc, CurMergeFile);
+						}
+						else
+						{
+							UE_LOG(OodleHandlerComponentLog, Error, TEXT("Could not read packet capture file: %s"), *CurMergeFile);
+
+							bAllMergeHandlesValid = false;
+							break;
+						}
+					}
+
+					FArchive* OutArc = (bAllMergeHandlesValid ? IFileManager::Get().CreateFileWriter(*OutputCapFile) : NULL);
+
+					if (bAllMergeHandlesValid && OutArc != NULL)
+					{
+						UE_LOG(OodleHandlerComponentLog, Log, TEXT("Merging files into output file: %s"), *OutputCapFile);
+
+						FPacketCaptureArchive OutputFile(*OutArc);
+						bool bErrorAppending = false;
+
+						OutputFile.SerializeCaptureHeader();
+
+						for (TMap<FArchive*, FString>::TConstIterator It(MergeMap); It; ++It)
+						{
+							FArchive* CurReadArc = It.Key();
+							FPacketCaptureArchive ReadFile(*CurReadArc);
+
+							UE_LOG(OodleHandlerComponentLog, Log, TEXT("    Merging: %s"), *It.Value());
+
+							OutputFile.AppendPacketFile(ReadFile);
+
+							if (OutputFile.IsError())
+							{
+								bErrorAppending = true;
+								break;
+							}
+
+							delete CurReadArc;
+						}
+
+						MergeMap.Empty();
+
+
+						if (bErrorAppending)
+						{
+							UE_LOG(OodleHandlerComponentLog, Error, TEXT("Error appending packet file."));
+						}
+
+
+						// @todo #JohnB: Log success
+
+						OutArc->Close();
+
+						delete OutArc;
+
+						OutArc = NULL;
+					}
+					else if (!bAllMergeHandlesValid)
+					{
+						UE_LOG(OodleHandlerComponentLog, Error, TEXT("Could not read all packet files."));
+					}
+					else // if (OutArc == NULL)
+					{
+						UE_LOG(OodleHandlerComponentLog, Error, TEXT("Could not create output file."));
+					}
+				}
+				else if (!bValidFileParms)
+				{
+					UE_LOG(OodleHandlerComponentLog, Error,
+							TEXT("The packet files must exist, and the output file must not already exist"));
+				}
+				else // if (MergeFiles.Num() <= 1)
+				{
+					UE_LOG(OodleHandlerComponentLog, Error, TEXT("You need to specify 2 or more packet files for merging"));
+				}
+			}
+			else
+			{
+				UE_LOG(OodleHandlerComponentLog, Error,
+						TEXT("Error parsing 'MergePackets' parameters. Must specify an output file, and packet files to merge"));
+			}
+		}
+		else
+		{
+			UE_LOG(OodleHandlerComponentLog, Error, TEXT("Unrecognized sub-command '%s'"), *MainCmd);
+		}
+	}
+	else
+	{
+		UE_LOG(OodleHandlerComponentLog, Error, TEXT("Error parsing main commandlet sub-command"));
+
+		// @todo #JohnB: Add a 'help' command, with documentation output
+	}
+#else
+	UE_LOG(OodleHandlerComponentLog, Error, TEXT("Oodle unavailable."));
+#endif
+
+
+	return (GWarn->Errors.Num() == 0 ? 0 : 1);
+}
+
+#undef LOCTEXT_NAMESPACE
+
diff --git a/Engine/Plugins/Runtime/PacketHandlers/CompressionComponents/Oodle/Source/OodleHandlerComponent/Public/OodleArchives.h b/Engine/Plugins/Runtime/PacketHandlers/CompressionComponents/Oodle/Source/OodleHandlerComponent/Public/OodleArchives.h
new file mode 100644
--- /dev/null
+++ b/Engine/Plugins/Runtime/PacketHandlers/CompressionComponents/Oodle/Source/OodleHandlerComponent/Public/OodleArchives.h
@@ -0,0 +1,83 @@
+// Copyright 1998-2015 Epic Games, Inc. All Rights Reserved.
+
+#pragma once
+
+#if HAS_OODLE_SDK
+#include "ArchiveBase.h"
+
+#define CAPTURE_HEADER_MAGIC 0x41091CC4
+#define CAPTURE_FILE_VERSION 0x0001
+
+
+// @todo #JohnB: Add to the header, whether the capture file is an input or output capture type,
+//					and use that to give a warning when merging mismatched input/output capture files
+//					(actually, make that error by default, requiring a special -blah commandline parm, to override the error)
+
+
+/**
+ * Archive for handling packet capture (.ucap) files
+ */
+class FPacketCaptureArchive : public FArchiveProxy
+{
+protected:
+	/** Capture file format version */
+	uint32 CaptureVersion;
+
+	/** Whether or not to flush immediately */
+	bool bImmediateFlush;
+
+
+public:
+	/**
+	 * Base constructor
+	 */
+	FPacketCaptureArchive(FArchive& InInnerArchive, bool bInDeleteInner=false)
+		: FArchiveProxy(InInnerArchive)
+	{
+		if (IsSaving())
+		{
+			CaptureVersion = CAPTURE_FILE_VERSION;
+		}
+		else
+		{
+			CaptureVersion = 0;
+		}
+
+		bImmediateFlush = FParse::Param(FCommandLine::Get(), TEXT("FORCELOGFLUSH"));
+	}
+
+	/**
+	 * Deletes the inner archive
+	 */
+	void DeleteInnerArchive()
+	{
+		delete (FArchive*)&InnerArchive;
+	}
+
+
+	/**
+	 * Serializes the file header, containing the file format UID (magic) and file version
+	 */
+	void SerializeCaptureHeader();
+
+	/**
+	 * Serialize an individual packet to/from the archive.
+	 * It is possible for there to be an incomplete packet stored - in which case, attempting to read will set the archives error mode
+	 *
+	 * @param PacketData	The location of/for the packet data
+	 * @param PacketSize	For saving, the size of the packet, for loading, inputs the size of the buffer, and outputs the packet size
+	 */
+	void SerializePacket(void* PacketData, uint32& PacketSize);
+
+
+	/**
+	 * Used for merging multiple packet files. Appends the specified packet file to this one.
+	 *
+	 * @param InPacketFile	The packet file archive, to append to this one.
+	 */
+	void AppendPacketFile(FPacketCaptureArchive& InPacketFile);
+};
+#endif
+
+
+
diff --git a/Engine/Plugins/Runtime/PacketHandlers/CompressionComponents/Oodle/Source/OodleHandlerComponent/Public/OodleHandlerComponent.h b/Engine/Plugins/Runtime/PacketHandlers/CompressionComponents/Oodle/Source/OodleHandlerComponent/Public/OodleHandlerComponent.h
--- a/Engine/Plugins/Runtime/PacketHandlers/CompressionComponents/Oodle/Source/OodleHandlerComponent/Public/OodleHandlerComponent.h
+++ b/Engine/Plugins/Runtime/PacketHandlers/CompressionComponents/Oodle/Source/OodleHandlerComponent/Public/OodleHandlerComponent.h
@@ -6,12 +6,70 @@
 #include "ModuleManager.h"
 #include "Core.h"
 
-#if HAS_OODLE_SDK == 1
+DECLARE_LOG_CATEGORY_EXTERN(OodleHandlerComponentLog, Log, All);
+
+
+#if HAS_OODLE_SDK
+#include "OodleArchives.h"
+
 #include "oodle.h"
 
-DECLARE_LOG_CATEGORY_EXTERN(OodleHandlerComponentLog, Log, All);
 
-// Symmetric Stream cipher
+// Whether or not to utilize the Oodle-example-code based training and packet capture code (in the process of deprecation)
+#define EXAMPLE_CAPTURE_FORMAT 0
+
+#if EXAMPLE_CAPTURE_FORMAT
+	#define CAPTURE_EXT TEXT(".bin")
+#else
+	#define CAPTURE_EXT TEXT(".ucap")
+#endif
+
+
+/** Globals */
+extern FString GOodleSaveDir;
+
+
+/**
+ * The mode that the Oodle packet handler should operate in
+ */
+enum EOodleHandlerMode
+{
+	Training,	// Stores packet captures for the server
+	Release		// Compresses packet data, based on the dictionary file
+};
+
+
+/**
+ * Stores dictionary data from an individual file (typically one dictionary each, for server/client)
+ *
+ * @todo #JohnB: In its current form, this is temporary/stopgap, based on the example code
+ */
+struct FOodleDictionary
+{
+	/** Base constructor */
+	FOodleDictionary()
+		: Dictionary(NULL)
+		, SharedDictionaryData(NULL)
+		, CompressorState(NULL)
+	{
+	}
+
+	/* Dictionary */
+	void* Dictionary; 
+
+	/* Shared Compressor */
+	OodleNetwork1_Shared* SharedDictionaryData;
+
+	/* State of Oodle */
+	OodleNetwork1UDP_State* CompressorState;
+};
+
+
+/**
+ * PacketHandler component for implementing Oodle support.
+ *
+ * Implementation uses trained/dictionary-based UDP compression.
+ */
 class OODLEHANDLERCOMPONENT_API OodleHandlerComponent : public HandlerComponent
 {
 public:
@@ -21,6 +79,17 @@ public:
 	/* Default Destructor */
 	~OodleHandlerComponent();
 
+	/**
+	 * Initializes first-run config settings
+	 */
+	static void InitFirstRunConfig();
+
+
+	/**
+	 * @todo #JohnB
+	 */
+	void InitializeDictionary(FString FilePath, FOodleDictionary& OutDictionary);
+
 	/* Initializes the handler component */
 	virtual void Initialize() override;
 
@@ -37,8 +106,20 @@ protected:
 	/* How many bytes have been saved from compressing the data */
 	uint32 BytesSaved;
 
+#if EXAMPLE_CAPTURE_FORMAT
 	/* File to log packets to */
-	IFileHandle* LogFile;
+	IFileHandle* PacketLogFile;
+
+	/** Whether or not to append data to a statically named log file */
+	bool bPacketLogAppend;
+#else
+	/* File to log input packets to */
+	FPacketCaptureArchive* InPacketLog;
+
+	/* File to log output packets to */
+	FPacketCaptureArchive* OutPacketLog;
+#endif
+
 
 	/* 
 	* Modes of the component
@@ -46,27 +127,34 @@ protected:
 	* [OodleHandlerComponent]
 	* Mode=Training
 	*/
-	enum {Training, Release}Mode;
+	EOodleHandlerMode Mode;
 
-	/* Dictionary */
-	void* Dictionary; 
 
-	/* Shared Compressor */
-	OodleNetwork1_Shared* SharedDictionaryData;
+	/** Server (Outgoing) dictionary data */
+	FOodleDictionary ServerDictionary;
 
-	/* State of Oodle */
-	OodleNetwork1UDP_State* CompressorState;
+	/** Client (Incoming - relative to server) dictionary data */
+	FOodleDictionary ClientDictionary;
 };
 
-DECLARE_DWORD_COUNTER_STAT_EXTERN( TEXT( "Oodle Bytes Raw" ), STAT_Oodle_BytesRaw, STATGROUP_Net, );
-DECLARE_DWORD_COUNTER_STAT_EXTERN( TEXT( "Oodle Bytes Compressed" ), STAT_Oodle_BytesCompressed, STATGROUP_Net, );
+DECLARE_DWORD_COUNTER_STAT_EXTERN(TEXT("Oodle Bytes Raw"), STAT_Oodle_BytesRaw, STATGROUP_Net, );
+DECLARE_DWORD_COUNTER_STAT_EXTERN(TEXT("Oodle Bytes Compressed"), STAT_Oodle_BytesCompressed, STATGROUP_Net, );
 #endif
 
 /* Reliability Module Interface */
 class FOodleComponentModuleInterface : public FPacketHandlerComponentModuleInterface
 {
+private:
+	/** Reference to the Oodle library handle */
+	void* OodleDllHandle;
+
 public:
-	virtual HandlerComponent* CreateComponentInstance() override;
+	FOodleComponentModuleInterface()
+		: OodleDllHandle(NULL)
+	{
+	}
+
+	virtual TSharedPtr<HandlerComponent> CreateComponentInstance(FString& Options) override;
 	virtual void StartupModule() override;
 	virtual void ShutdownModule() override;
-};
\ No newline at end of file
+};
