Mergine engine changes to //UE4/Main @ 2748007

diff --git a/Engine/Plugins/2D/Paper2D/Source/Paper2DEditor/Private/MeshPainting/MeshPaintSpriteAdapter.h b/Engine/Plugins/2D/Paper2D/Source/Paper2DEditor/Private/MeshPainting/MeshPaintSpriteAdapter.h
--- a/Engine/Plugins/2D/Paper2D/Source/Paper2DEditor/Private/MeshPainting/MeshPaintSpriteAdapter.h
+++ b/Engine/Plugins/2D/Paper2D/Source/Paper2DEditor/Private/MeshPainting/MeshPaintSpriteAdapter.h
@@ -16,6 +16,7 @@ public:
 	static void InitializeAdapterGlobals() {}
 
 	virtual bool Construct(UMeshComponent* InComponent, int32 InPaintingMeshLODIndex, int32 InUVChannelIndex) override;
+	virtual bool InitializeMeshData() override { return true; }
 	virtual void OnAdded() override {}
 	virtual void OnRemoved() override {}
 	virtual int32 GetNumTexCoords() const override;
diff --git a/Engine/Plugins/Runtime/PacketHandlers/CompressionComponents/Oodle/FirstRun.bat b/Engine/Plugins/Runtime/PacketHandlers/CompressionComponents/Oodle/FirstRun.bat
--- a/Engine/Plugins/Runtime/PacketHandlers/CompressionComponents/Oodle/FirstRun.bat
+++ b/Engine/Plugins/Runtime/PacketHandlers/CompressionComponents/Oodle/FirstRun.bat
@@ -30,7 +30,7 @@ echo.
 
 :EnablePlugin
 set EnableCommandletParms=-run=plugin enable Oodle
-set FinalEnableCmdLine=%GameName:"=% %EnableCommandletParms%
+set FinalEnableCmdLine=%GameName:"=% %EnableCommandletParms% -forcelogflush
 
 echo Executing plugin enable commandlet - commandline:
 echo %FinalEnableCmdLine%
@@ -49,7 +49,7 @@ pause
 
 :EnableHandler
 set HandlerCommandletParms=-run=OodleHandlerComponent.OodleTrainerCommandlet enable
-set FinalHandlerCmdLine=%GameName:"=% %HandlerCommandletParms%
+set FinalHandlerCmdLine=%GameName:"=% %HandlerCommandletParms% -forcelogflush
 
 
 echo Executing Oodle PacketHandler enable commandlet - commandline:
diff --git a/Engine/Plugins/Runtime/PacketHandlers/CompressionComponents/Oodle/Source/OodleHandlerComponent/Private/OodleHandlerComponent.cpp b/Engine/Plugins/Runtime/PacketHandlers/CompressionComponents/Oodle/Source/OodleHandlerComponent/Private/OodleHandlerComponent.cpp
--- a/Engine/Plugins/Runtime/PacketHandlers/CompressionComponents/Oodle/Source/OodleHandlerComponent/Private/OodleHandlerComponent.cpp
+++ b/Engine/Plugins/Runtime/PacketHandlers/CompressionComponents/Oodle/Source/OodleHandlerComponent/Private/OodleHandlerComponent.cpp
@@ -6,6 +6,8 @@
 DEFINE_LOG_CATEGORY(OodleHandlerComponentLog);
 
 
+// @todo #JohnB: Right now, there is no caching of loaded dictionary files, so this is very inefficient with multiple clients
+
 #if HAS_OODLE_SDK
 #define OODLE_INI_SECTION TEXT("OodleHandlerComponent")
 
@@ -79,6 +81,10 @@ OodleHandlerComponent::OodleHandlerComponent()
 #else
 	, InPacketLog(NULL)
 	, OutPacketLog(NULL)
+#endif
+	, bEnableOodle(false)
+#if !UE_BUILD_SHIPPING
+	, bUseDictionaryIfPresent(false)
 #endif
 	, Mode(EOodleHandlerMode::Release)
 	, ServerDictionary()
@@ -154,6 +160,12 @@ void OodleHandlerComponent::InitFirstRunConfig()
 	// Check that the OodleHandlerComponent section exists, and if not, init with defaults
 	if (!GConfig->DoesSectionExist(OODLE_INI_SECTION, GEngineIni))
 	{
+		GConfig->SetBool(OODLE_INI_SECTION, TEXT("bEnableOodle"), true, GEngineIni);
+
+#if !UE_BUILD_SHIPPING
+		GConfig->SetBool(OODLE_INI_SECTION, TEXT("bUseDictionaryIfPresent"), false, GEngineIni);
+#endif
+
 		GConfig->SetString(OODLE_INI_SECTION, TEXT("Mode"), TEXT("Training"), GEngineIni);
 		GConfig->SetString(OODLE_INI_SECTION, TEXT("PacketLogFile"), TEXT("PacketDump"), GEngineIni);
 
@@ -170,8 +182,20 @@ void OodleHandlerComponent::Initialize()
 {
 	InitFirstRunConfig();
 
-#if EXAMPLE_CAPTURE_FORMAT
 	// Class config variables
+	GConfig->GetBool(OODLE_INI_SECTION, TEXT("bEnableOodle"), bEnableOodle, GEngineIni);
+
+	if (!bEnableOodle && FParse::Param(FCommandLine::Get(), TEXT("Oodle")))
+	{
+		UE_LOG(OodleHandlerComponentLog, Log, TEXT("Force-enabling Oodle from commandline."));
+		bEnableOodle = true;
+	}
+
+#if !UE_BUILD_SHIPPING
+	GConfig->GetBool(OODLE_INI_SECTION, TEXT("bUseDictionaryIfPresent"), bUseDictionaryIfPresent, GEngineIni);
+#endif
+
+#if EXAMPLE_CAPTURE_FORMAT
 	GConfig->GetBool(OODLE_INI_SECTION, TEXT("bPacketLogAppend"), bPacketLogAppend, GEngineIni);
 #endif
 
@@ -193,158 +217,181 @@ void OodleHandlerComponent::Initialize()
 		Mode = EOodleHandlerMode::Release;
 	}
 
-	switch(Mode)
+	if (bEnableOodle)
 	{
-	case EOodleHandlerMode::Training:
+#if !UE_BUILD_SHIPPING
+		bool bForceTraining = FParse::Param(FCommandLine::Get(), TEXT("OodleTraining"));
+
+		if (bForceTraining && Mode != EOodleHandlerMode::Training)
+		{
+			UE_LOG(OodleHandlerComponentLog, Log, TEXT("Switching Oodle to Training mode."));
+
+			Mode = EOodleHandlerMode::Training;
+		}
+		// Forces release mode, if configured to search for (and finds) fallback dictionaries
+		else if (!bForceTraining && bUseDictionaryIfPresent && Mode != EOodleHandlerMode::Release)
 		{
-			if (Handler->Mode == Handler::Mode::Server)
+			FString ServerDicPath;
+			FString ClientDicPath;
+
+			if (FindFallbackDictionaries(ServerDicPath, ClientDicPath, true))
 			{
-				IPlatformFile& PlatformFile = FPlatformFileManager::Get().GetPlatformFile();
-				FString ReadOutputLogDirectory = FPaths::Combine(*GOodleSaveDir, TEXT("Server"));
-				FString BaseFilename;
+				UE_LOG(OodleHandlerComponentLog, Log,
+						TEXT("Switching Oodle to Release mode - found server dictionary '%s' and client dictionary '%s'."),
+						*ServerDicPath, *ClientDicPath);
 
-				PlatformFile.CreateDirectoryTree(*ReadOutputLogDirectory);
-				PlatformFile.CreateDirectoryTree(*FPaths::Combine(*ReadOutputLogDirectory, TEXT("Input")));
-				PlatformFile.CreateDirectoryTree(*FPaths::Combine(*ReadOutputLogDirectory, TEXT("Output")));
-				GConfig->GetString(OODLE_INI_SECTION, TEXT("PacketLogFile"), BaseFilename, GEngineIni);
+				Mode = EOodleHandlerMode::Release;
+			}
+		}
+#endif
 
-				BaseFilename = FPaths::GetBaseFilename(BaseFilename);
+		switch(Mode)
+		{
+		case EOodleHandlerMode::Training:
+			{
+				if (Handler->Mode == Handler::Mode::Server)
+				{
+					IPlatformFile& PlatformFile = FPlatformFileManager::Get().GetPlatformFile();
+					FString ReadOutputLogDirectory = FPaths::Combine(*GOodleSaveDir, TEXT("Server"));
+					FString BaseFilename;
 
+					PlatformFile.CreateDirectoryTree(*ReadOutputLogDirectory);
+					PlatformFile.CreateDirectoryTree(*FPaths::Combine(*ReadOutputLogDirectory, TEXT("Input")));
+					PlatformFile.CreateDirectoryTree(*FPaths::Combine(*ReadOutputLogDirectory, TEXT("Output")));
+					GConfig->GetString(OODLE_INI_SECTION, TEXT("PacketLogFile"), BaseFilename, GEngineIni);
 
-#if EXAMPLE_CAPTURE_FORMAT
-				FString FinalFilePath;
+					BaseFilename = FPaths::GetBaseFilename(BaseFilename);
 
-				// When bPacketLogAppend is used, a static log filename is used
-				if (bPacketLogAppend)
-				{
-					FString PreExtFilePath = FPaths::Combine(*ReadOutputLogDirectory, *BaseFilename);
 
-					// Try0: PacketDump.ucap, Try1: PacketDump_1.ucap, Try#:PacketDump_#.ucap
-					FinalFilePath = PreExtFilePath + CAPTURE_EXT;
+#if EXAMPLE_CAPTURE_FORMAT
+					FString FinalFilePath;
 
-					for (int32 i=1; ; i++)
+					// When bPacketLogAppend is used, a static log filename is used
+					if (bPacketLogAppend)
 					{
-						// If not appending to a file, make sure to disable this bool - otherwise the file header won't be written
-						bPacketLogAppend = PlatformFile.FileExists(*FinalFilePath);
-						PacketLogFile = PlatformFile.OpenWrite(*FinalFilePath, true);
+						FString PreExtFilePath = FPaths::Combine(*ReadOutputLogDirectory, *BaseFilename);
 
-						if (PacketLogFile != NULL)
+						// Try0: PacketDump.ucap, Try1: PacketDump_1.ucap, Try#:PacketDump_#.ucap
+						FinalFilePath = PreExtFilePath + CAPTURE_EXT;
+
+						for (int32 i=1; ; i++)
 						{
-							break;
-						}
+							// If not appending to a file, make sure to disable this bool - otherwise the file header won't be written
+							bPacketLogAppend = PlatformFile.FileExists(*FinalFilePath);
+							PacketLogFile = PlatformFile.OpenWrite(*FinalFilePath, true);
+
+							if (PacketLogFile != NULL)
+							{
+								break;
+							}
 
 
-						FinalFilePath = PreExtFilePath + FString::Printf(TEXT("_%i"), i) + CAPTURE_EXT;
+							FinalFilePath = PreExtFilePath + FString::Printf(TEXT("_%i"), i) + CAPTURE_EXT;
+						}
 					}
-				}
-				// When non-append logging is used, the log filename uses a timestamp
-				else
-				{
-					BaseFilename = BaseFilename + TEXT("_") + FDateTime::Now().ToString();
+					// When non-append logging is used, the log filename uses a timestamp
+					else
+					{
+						BaseFilename = BaseFilename + TEXT("_") + FDateTime::Now().ToString();
 
-					FString PreExtFilePath = FPaths::Combine(*ReadOutputLogDirectory, *BaseFilename);
-					FinalFilePath = PreExtFilePath + CAPTURE_EXT;
+						FString PreExtFilePath = FPaths::Combine(*ReadOutputLogDirectory, *BaseFilename);
+						FinalFilePath = PreExtFilePath + CAPTURE_EXT;
 
-					for (int32 i=1; PlatformFile.FileExists(*FinalFilePath); i++)
-					{
-						FinalFilePath = PreExtFilePath + FString::Printf(TEXT("_%i"), i) + CAPTURE_EXT;
-					}
+						for (int32 i=1; PlatformFile.FileExists(*FinalFilePath); i++)
+						{
+							FinalFilePath = PreExtFilePath + FString::Printf(TEXT("_%i"), i) + CAPTURE_EXT;
+						}
 
-					// Try0: PacketDump_2015.09.23-20.12.54.ucap, Try1: PacketDump_2015.09.23-20.12.54_1.ucap, ...
-					PacketLogFile = PlatformFile.OpenWrite(*FinalFilePath);
-				}
+						// Try0: PacketDump_2015.09.23-20.12.54.ucap, Try1: PacketDump_2015.09.23-20.12.54_1.ucap, ...
+						PacketLogFile = PlatformFile.OpenWrite(*FinalFilePath);
+					}
 
-				if (PacketLogFile == NULL)
-				{
-					LowLevelFatalError(TEXT("Failed to create file %s"), *FinalFilePath);
-					return;
-				}
+					if (PacketLogFile == NULL)
+					{
+						LowLevelFatalError(TEXT("Failed to create file %s"), *FinalFilePath);
+						return;
+					}
 #else
-				BaseFilename = BaseFilename + TEXT("_") + FDateTime::Now().ToString();
+					BaseFilename = BaseFilename + TEXT("_") + FDateTime::Now().ToString();
 
-				FString PreExtInFilePath = FPaths::Combine(*ReadOutputLogDirectory, TEXT("Input"), *(BaseFilename + TEXT("_Input")));
-				FString PreExtOutFilePath = FPaths::Combine(*ReadOutputLogDirectory, TEXT("Output"), *(BaseFilename + TEXT("_Output")));
-				FString InPath = PreExtInFilePath + CAPTURE_EXT;
-				FString OutPath = PreExtOutFilePath + CAPTURE_EXT;
+					FString PreExtInFilePath = FPaths::Combine(*ReadOutputLogDirectory, TEXT("Input"), *(BaseFilename + TEXT("_Input")));
+					FString PreExtOutFilePath = FPaths::Combine(*ReadOutputLogDirectory, TEXT("Output"), *(BaseFilename + TEXT("_Output")));
+					FString InPath = PreExtInFilePath + CAPTURE_EXT;
+					FString OutPath = PreExtOutFilePath + CAPTURE_EXT;
 
-				// Ensure the In/Out filenames are unique
-				for (int32 i=1; PlatformFile.FileExists(*InPath) || PlatformFile.FileExists(*OutPath); i++)
-				{
-					InPath = PreExtInFilePath + FString::Printf(TEXT("_%i"), i) + CAPTURE_EXT;
-					OutPath = PreExtOutFilePath + FString::Printf(TEXT("_%i"), i) + CAPTURE_EXT;
-				}
+					// Ensure the In/Out filenames are unique
+					for (int32 i=1; PlatformFile.FileExists(*InPath) || PlatformFile.FileExists(*OutPath); i++)
+					{
+						InPath = PreExtInFilePath + FString::Printf(TEXT("_%i"), i) + CAPTURE_EXT;
+						OutPath = PreExtOutFilePath + FString::Printf(TEXT("_%i"), i) + CAPTURE_EXT;
+					}
 
-				FArchive* InArc = IFileManager::Get().CreateFileWriter(*InPath);
-				FArchive* OutArc = (InArc != NULL ? IFileManager::Get().CreateFileWriter(*OutPath) : NULL);
+					FArchive* InArc = IFileManager::Get().CreateFileWriter(*InPath);
+					FArchive* OutArc = (InArc != NULL ? IFileManager::Get().CreateFileWriter(*OutPath) : NULL);
 
-				InPacketLog = (InArc != NULL ? new FPacketCaptureArchive(*InArc) : NULL);
-				OutPacketLog = (OutArc != NULL ? new FPacketCaptureArchive(*OutArc) : NULL);
+					InPacketLog = (InArc != NULL ? new FPacketCaptureArchive(*InArc) : NULL);
+					OutPacketLog = (OutArc != NULL ? new FPacketCaptureArchive(*OutArc) : NULL);
 
 
-				if (InPacketLog == NULL || OutPacketLog == NULL)
-				{
-					LowLevelFatalError(TEXT("Failed to create files '%s' and '%s'"), *InPath, *OutPath);
-					return;
-				}
+					if (InPacketLog == NULL || OutPacketLog == NULL)
+					{
+						LowLevelFatalError(TEXT("Failed to create files '%s' and '%s'"), *InPath, *OutPath);
+						return;
+					}
 #endif
 
 
 #if EXAMPLE_CAPTURE_FORMAT
-				// Only write the initial header, if not appending
-				if (!bPacketLogAppend)
-				{
-					int32 NumberOfChannels = 1;
+					// Only write the initial header, if not appending
+					if (!bPacketLogAppend)
+					{
+						int32 NumberOfChannels = 1;
 			
-					PacketLogFile->Write(reinterpret_cast<const uint8*>(&NumberOfChannels), sizeof(NumberOfChannels));
-				}
+						PacketLogFile->Write(reinterpret_cast<const uint8*>(&NumberOfChannels), sizeof(NumberOfChannels));
+					}
 #else
-				InPacketLog->SerializeCaptureHeader();
-				OutPacketLog->SerializeCaptureHeader();
+					InPacketLog->SerializeCaptureHeader();
+					OutPacketLog->SerializeCaptureHeader();
 #endif
-			}
-
-			break;
-		}
-
-	case EOodleHandlerMode::Release:
-		{
-			FString ServerDictionaryPath;
-			FString ClientDictionaryPath;
+				}
 
-			GConfig->GetString(OODLE_INI_SECTION, TEXT("ServerDictionary"), ServerDictionaryPath, GEngineIni);
-			GConfig->GetString(OODLE_INI_SECTION, TEXT("ClientDictionary"), ClientDictionaryPath, GEngineIni);
+				break;
+			}
 
-			if (ServerDictionaryPath.Len() < 0 || ClientDictionaryPath.Len() < 0)
+		case EOodleHandlerMode::Release:
 			{
-				LowLevelFatalError(TEXT("Specify both Server/Client dictionaries for Oodle compressor in DefaultEngine.ini"));
-				return;
-			}
+				FString ServerDictionaryPath;
+				FString ClientDictionaryPath;
+				bool bGotDictionaryPath = false;
 
+#if !UE_BUILD_SHIPPING
+				if (bUseDictionaryIfPresent)
+				{
+					bGotDictionaryPath = FindFallbackDictionaries(ServerDictionaryPath, ClientDictionaryPath);
+				}
+#endif
 
-			// Path must be within game directory, e.g: "Content/Oodle/Output.udic" becomes "ShooterGam/Content/Oodle/Output.udic"
-			ServerDictionaryPath = FPaths::Combine(*FPaths::GameDir(), *ServerDictionaryPath);
-			ClientDictionaryPath = FPaths::Combine(*FPaths::GameDir(), *ClientDictionaryPath);
+				if (!bGotDictionaryPath)
+				{
+					bGotDictionaryPath = GetDictionaryPaths(ServerDictionaryPath, ClientDictionaryPath);
+				}
 
-			FPaths::CollapseRelativeDirectories(ServerDictionaryPath);
-			FPaths::CollapseRelativeDirectories(ClientDictionaryPath);
+				if (bGotDictionaryPath)
+				{
+					InitializeDictionary(ServerDictionaryPath, ServerDictionary);
+					InitializeDictionary(ClientDictionaryPath, ClientDictionary);
+				}
+				else
+				{
+					LowLevelFatalError(TEXT("Failed to load Oodle dictionaries."));
+				}
 
-			// Don't allow directory traversal to escape the game directory
-			if (!ServerDictionaryPath.StartsWith(FPaths::GameDir()) || !ClientDictionaryPath.StartsWith(FPaths::GameDir()))
-			{
-				LowLevelFatalError(TEXT("DictionaryFile not allowed to use ../ paths to escape game directory."));
-				return;
+				break;
 			}
 
-
-			InitializeDictionary(ServerDictionaryPath, ServerDictionary);
-			InitializeDictionary(ClientDictionaryPath, ClientDictionary);
-
-
+		default:
 			break;
 		}
-
-	default:
-		break;
 	}
 
 	Initialized();
@@ -430,7 +477,8 @@ void OodleHandlerComponent::InitializeDictionary(FString FilePath, FOodleDiction
 
 			OutDictionary.SharedDictionaryData = (OodleNetwork1_Shared *)malloc(SharedSize);
 			check(OutDictionary.SharedDictionaryData != NULL);
-			OodleNetwork1_Shared_SetWindow(OutDictionary.SharedDictionaryData, StateHTBits, OutDictionary.Dictionary, (S32)DictionarySize);
+			OodleNetwork1_Shared_SetWindow(OutDictionary.SharedDictionaryData, StateHTBits, OutDictionary.Dictionary,
+											(S32)DictionarySize);
 
 			//-----------------------------------------------------
 	
@@ -517,7 +565,8 @@ void OodleHandlerComponent::InitializeDictionary(FString FilePath, FOodleDiction
 
 			OutDictionary.SharedDictionaryData = (OodleNetwork1_Shared *)malloc(SharedSize);
 			check(OutDictionary.SharedDictionaryData != NULL);
-			OodleNetwork1_Shared_SetWindow(OutDictionary.SharedDictionaryData, StateHTBits, OutDictionary.Dictionary, (S32)DictionarySize);
+			OodleNetwork1_Shared_SetWindow(OutDictionary.SharedDictionaryData, StateHTBits, OutDictionary.Dictionary,
+											(S32)DictionarySize);
 
 			//-----------------------------------------------------
 	
@@ -541,6 +590,160 @@ void OodleHandlerComponent::InitializeDictionary(FString FilePath, FOodleDiction
 	}
 }
 
+bool OodleHandlerComponent::GetDictionaryPaths(FString& OutServerDictionary, FString& OutClientDictionary, bool bFailFatal/*=true*/)
+{
+	bool bSuccess = false;
+
+	FString ServerDictionaryPath;
+	FString ClientDictionaryPath;
+
+	bSuccess = GConfig->GetString(OODLE_INI_SECTION, TEXT("ServerDictionary"), ServerDictionaryPath, GEngineIni);
+	bSuccess = bSuccess && GConfig->GetString(OODLE_INI_SECTION, TEXT("ClientDictionary"), ClientDictionaryPath, GEngineIni);
+
+	if (bSuccess && (ServerDictionaryPath.Len() < 0 || ClientDictionaryPath.Len() < 0))
+	{
+		const TCHAR* Msg = TEXT("Specify both Server/Client dictionaries for Oodle compressor in DefaultEngine.ini");
+
+		if (bFailFatal)
+		{
+			LowLevelFatalError(TEXT("%s"), Msg);
+		}
+		else
+		{
+			UE_LOG(OodleHandlerComponentLog, Warning, TEXT("%s"), Msg);
+		}
+
+		bSuccess = false;
+	}
+
+	if (bSuccess)
+	{
+		// Path must be within game directory, e.g: "Content/Oodle/Output.udic" becomes "ShooterGam/Content/Oodle/Output.udic"
+		ServerDictionaryPath = FPaths::Combine(*FPaths::GameDir(), *ServerDictionaryPath);
+		ClientDictionaryPath = FPaths::Combine(*FPaths::GameDir(), *ClientDictionaryPath);
+
+		FPaths::CollapseRelativeDirectories(ServerDictionaryPath);
+		FPaths::CollapseRelativeDirectories(ClientDictionaryPath);
+
+
+		// Don't allow directory traversal to escape the game directory
+		if (!ServerDictionaryPath.StartsWith(FPaths::GameDir()) || !ClientDictionaryPath.StartsWith(FPaths::GameDir()))
+		{
+			const TCHAR* Msg = TEXT("DictionaryFile not allowed to use ../ paths to escape game directory.");
+
+			if (bFailFatal)
+			{
+				LowLevelFatalError(TEXT("%s"), Msg);
+			}
+			else
+			{
+				UE_LOG(OodleHandlerComponentLog, Warning, TEXT("%s"), Msg);
+			}
+
+			bSuccess = false;
+		}
+	}
+
+	if (bSuccess)
+	{
+		OutServerDictionary = ServerDictionaryPath;
+		OutClientDictionary = ClientDictionaryPath;
+	}
+	else
+	{
+		OutServerDictionary = TEXT("");
+		OutClientDictionary = TEXT("");
+	}
+
+
+	return bSuccess;
+}
+
+#if !UE_BUILD_SHIPPING
+bool OodleHandlerComponent::FindFallbackDictionaries(FString& OutServerDictionary, FString& OutClientDictionary, bool bTestOnly/*=false*/)
+{
+	bool bSuccess = false;
+
+	OutServerDictionary = TEXT("");
+	OutClientDictionary = TEXT("");
+
+	// First test the normal dictionary config paths
+	FString DefaultServerDicPath;
+	FString DefaultClientDicPath;
+	IFileManager& FileMan = IFileManager::Get();
+
+	bSuccess = GetDictionaryPaths(DefaultServerDicPath, DefaultClientDicPath, false);
+
+	bSuccess = bSuccess && FileMan.FileExists(*DefaultServerDicPath);
+	bSuccess = bSuccess && FileMan.FileExists(*DefaultClientDicPath);
+
+
+	// If either of the default dictionaries do not exist, do a more speculative search
+	if (!bSuccess)
+	{
+		TArray<FString> DictionaryList;
+		
+		FileMan.FindFilesRecursive(DictionaryList, *FPaths::GameDir(), TEXT("*.udic"), true, false);
+
+		if (DictionaryList.Num() > 0)
+		{
+			// Sort the list alphabetically (case-insensitive)
+			DictionaryList.Sort();
+
+			// Very simple matching - anything 'server/output' is a server dictionary, anything 'client/input' is a client dictionary
+			int32 FoundServerIdx = DictionaryList.IndexOfByPredicate(
+				[](const FString& CurEntry)
+				{
+					return CurEntry.Contains(TEXT("Server")) || CurEntry.Contains(TEXT("Output"));
+				});
+
+			int32 FoundClientIdx = DictionaryList.IndexOfByPredicate(
+				[](const FString& CurEntry)
+				{
+					return CurEntry.Contains(TEXT("Client")) || CurEntry.Contains(TEXT("Input"));
+				});
+
+			bSuccess = FoundServerIdx != INDEX_NONE && FoundClientIdx != INDEX_NONE;
+
+			if (!bTestOnly)
+			{
+				UE_LOG(OodleHandlerComponentLog, Log,
+						TEXT("Searched for Oodle dictionary files, and selected the following non-default dictionaries:"));
+			}
+
+			if (bSuccess)
+			{
+				OutServerDictionary = DictionaryList[FoundServerIdx];
+				OutClientDictionary = DictionaryList[FoundClientIdx];
+			}
+			else
+			{
+				// If all else fails, use any found dictionary, or just use the first listed dictionary, for both client/server
+				int32 DicIdx = FMath::Max3(0, FoundServerIdx, FoundClientIdx);
+
+				OutServerDictionary = DictionaryList[DicIdx];
+				OutClientDictionary = DictionaryList[DicIdx];
+
+				bSuccess = true;
+
+				if (!bTestOnly)
+				{
+					UE_LOG(OodleHandlerComponentLog, Log, TEXT("WARNING: Using the same dictionary for both server/client!"));
+				}
+			}
+
+			if (!bTestOnly)
+			{
+				UE_LOG(OodleHandlerComponentLog, Log, TEXT("   Server: %s"), *OutServerDictionary);
+				UE_LOG(OodleHandlerComponentLog, Log, TEXT("   Client: %s"), *OutClientDictionary);
+			}
+		}
+	}
+
+	return bSuccess;
+}
+#endif
+
 bool OodleHandlerComponent::IsValid() const
 {
 	return true;
@@ -548,125 +751,131 @@ bool OodleHandlerComponent::IsValid() const
 
 void OodleHandlerComponent::Incoming(FBitReader& Packet)
 {
-	switch(Mode)
+	if (bEnableOodle)
 	{
-	case EOodleHandlerMode::Training:
+		switch(Mode)
 		{
-			if(Handler->Mode == Handler::Mode::Server)
+		case EOodleHandlerMode::Training:
 			{
-				uint32 SizeOfPacket = Packet.GetNumBytes();
+				if(Handler->Mode == Handler::Mode::Server)
+				{
+					uint32 SizeOfPacket = Packet.GetNumBytes();
 
 #if EXAMPLE_CAPTURE_FORMAT
-				int32 channel = 0;
+					int32 channel = 0;
 
-				PacketLogFile->Write(reinterpret_cast<const uint8*>(&channel), sizeof(channel));
-				PacketLogFile->Write(reinterpret_cast<const uint8*>(&SizeOfPacket), sizeof(SizeOfPacket));
-				PacketLogFile->Write(reinterpret_cast<const uint8*>(Packet.GetData()), SizeOfPacket);
+					PacketLogFile->Write(reinterpret_cast<const uint8*>(&channel), sizeof(channel));
+					PacketLogFile->Write(reinterpret_cast<const uint8*>(&SizeOfPacket), sizeof(SizeOfPacket));
+					PacketLogFile->Write(reinterpret_cast<const uint8*>(Packet.GetData()), SizeOfPacket);
 #else
-				void* PacketData = Packet.GetData();
-				InPacketLog->SerializePacket(PacketData, SizeOfPacket);
+					void* PacketData = Packet.GetData();
+					InPacketLog->SerializePacket(PacketData, SizeOfPacket);
 #endif
+				}
+				break;
 			}
-			break;
-		}
 
-	case EOodleHandlerMode::Release:
-		{
-			uint32 DecompressedLength;
-			Packet.SerializeIntPacked(DecompressedLength);
-
-			if (ensure(DecompressedLength < MAX_OODLE_PACKET_SIZE))
+		case EOodleHandlerMode::Release:
 			{
-				// @todo #JohnB: You want to tweak the size of 'MAX_OODLE_PACKET_SIZE'
-				static uint8 CompressedData[MAX_OODLE_PACKET_SIZE];
-				static uint8 DecompressedData[MAX_OODLE_PACKET_SIZE];
-				const int32 CompressedLength = Packet.GetBytesLeft();
-				FOodleDictionary* CurDict = ((Handler->Mode == Handler::Mode::Server) ? &ClientDictionary : &ServerDictionary);
+				uint32 DecompressedLength;
+				Packet.SerializeIntPacked(DecompressedLength);
 
-				Packet.Serialize(CompressedData, CompressedLength);
+				if (ensure(DecompressedLength < MAX_OODLE_PACKET_SIZE))
+				{
+					// @todo #JohnB: You want to tweak the size of 'MAX_OODLE_PACKET_SIZE'
+					static uint8 CompressedData[MAX_OODLE_PACKET_SIZE];
+					static uint8 DecompressedData[MAX_OODLE_PACKET_SIZE];
+					const int32 CompressedLength = Packet.GetBytesLeft();
+					FOodleDictionary* CurDict = ((Handler->Mode == Handler::Mode::Server) ? &ClientDictionary : &ServerDictionary);
 
-				check(OodleNetwork1UDP_Decode(CurDict->CompressorState, CurDict->SharedDictionaryData, CompressedData, CompressedLength,
-												DecompressedData, DecompressedLength));
+					Packet.Serialize(CompressedData, CompressedLength);
 
+					check(OodleNetwork1UDP_Decode(CurDict->CompressorState, CurDict->SharedDictionaryData, CompressedData,
+													CompressedLength, DecompressedData, DecompressedLength));
 
-				// @todo #JohnB: Decompress directly into FBitReader's buffer
-				FBitReader UnCompressedPacket(DecompressedData, DecompressedLength * 8);
-				Packet = UnCompressedPacket;
-			}
 
+					// @todo #JohnB: Decompress directly into FBitReader's buffer
+					FBitReader UnCompressedPacket(DecompressedData, DecompressedLength * 8);
+					Packet = UnCompressedPacket;
+				}
 
-			break;
+
+				break;
+			}
 		}
 	}
 }
 
 void OodleHandlerComponent::Outgoing(FBitWriter& Packet)
 {
-	switch(Mode)
+	if (bEnableOodle)
 	{
-	case EOodleHandlerMode::Training:
+		switch(Mode)
 		{
-			if(Handler->Mode == Handler::Mode::Server)
+		case EOodleHandlerMode::Training:
 			{
-				uint32 SizeOfPacket = Packet.GetNumBytes();
+				if(Handler->Mode == Handler::Mode::Server)
+				{
+					uint32 SizeOfPacket = Packet.GetNumBytes();
 
 #if EXAMPLE_CAPTURE_FORMAT
-				int32 channel = 0;
+					int32 channel = 0;
 
-				PacketLogFile->Write(reinterpret_cast<const uint8*>(&channel), sizeof(channel));
-				PacketLogFile->Write(reinterpret_cast<const uint8*>(&SizeOfPacket), sizeof(SizeOfPacket));
-				PacketLogFile->Write(reinterpret_cast<const uint8*>(Packet.GetData()), SizeOfPacket);
+					PacketLogFile->Write(reinterpret_cast<const uint8*>(&channel), sizeof(channel));
+					PacketLogFile->Write(reinterpret_cast<const uint8*>(&SizeOfPacket), sizeof(SizeOfPacket));
+					PacketLogFile->Write(reinterpret_cast<const uint8*>(Packet.GetData()), SizeOfPacket);
 #else
-				OutPacketLog->SerializePacket((void*)Packet.GetData(), SizeOfPacket);
+					OutPacketLog->SerializePacket((void*)Packet.GetData(), SizeOfPacket);
 #endif
+				}
+				break;
 			}
-			break;
-		}
 
-	case EOodleHandlerMode::Release:
-		{
-			// Add size
-			uint32 UncompressedLength = Packet.GetNumBytes();
+		case EOodleHandlerMode::Release:
+			{
+				// Add size
+				uint32 UncompressedLength = Packet.GetNumBytes();
 
-			check(UncompressedLength < MAX_OODLE_PACKET_SIZE);
+				check(UncompressedLength < MAX_OODLE_PACKET_SIZE);
 
-			if (UncompressedLength > 0)
-			{
-				// @todo #JohnB: You should be able to share the same two buffers between Incoming and Outgoing, as an optimization
-				static uint8 UncompressedData[MAX_OODLE_PACKET_SIZE];
-				static uint8 CompressedData[MAX_OODLE_PACKET_SIZE];
-				FOodleDictionary* CurDict = ((Handler->Mode == Handler::Mode::Server) ? &ServerDictionary : &ClientDictionary);
+				if (UncompressedLength > 0)
+				{
+					// @todo #JohnB: You should be able to share the same two buffers between Incoming and Outgoing, as an optimization
+					static uint8 UncompressedData[MAX_OODLE_PACKET_SIZE];
+					static uint8 CompressedData[MAX_OODLE_PACKET_SIZE];
+					FOodleDictionary* CurDict = ((Handler->Mode == Handler::Mode::Server) ? &ServerDictionary : &ClientDictionary);
 
-				check(OodleLZ_GetCompressedBufferSizeNeeded(UncompressedLength) < MAX_OODLE_PACKET_SIZE);
+					check(OodleLZ_GetCompressedBufferSizeNeeded(UncompressedLength) < MAX_OODLE_PACKET_SIZE);
 
-				memcpy(UncompressedData, Packet.GetData(), UncompressedLength);
+					memcpy(UncompressedData, Packet.GetData(), UncompressedLength);
 
-				SINTa CompressedLengthSINT = OodleNetwork1UDP_Encode(CurDict->CompressorState, CurDict->SharedDictionaryData,
-																		UncompressedData, UncompressedLength, CompressedData);
+					SINTa CompressedLengthSINT = OodleNetwork1UDP_Encode(CurDict->CompressorState, CurDict->SharedDictionaryData,
+																			UncompressedData, UncompressedLength, CompressedData);
 
-				uint32 CompressedLength = static_cast<uint32>(CompressedLengthSINT);
+					uint32 CompressedLength = static_cast<uint32>(CompressedLengthSINT);
 
-				Packet.Reset();
+					Packet.Reset();
 
-				// @todo #JohnB: Compress directly into a (deliberately oversized) FBitWriter buffer, which you can shrink afterwards
-				Packet.SerializeIntPacked(UncompressedLength);
-				Packet.Serialize(CompressedData, CompressedLength);
+					// @todo #JohnB: Compress directly into a (deliberately oversized) FBitWriter buffer, which you can shrink afterwards
+					Packet.SerializeIntPacked(UncompressedLength);
+					Packet.Serialize(CompressedData, CompressedLength);
 
 
-				BytesSaved += UncompressedLength - CompressedLength;
+					BytesSaved += UncompressedLength - CompressedLength;
 
-				//UE_LOG(OodleHandlerComponentLog, Log, TEXT("Oodle Compressed: UnCompressed: %i Compressed: %i Total: %i"),
-				//			UncompressedLength, CompressedLength, BytesSaved);
+					//UE_LOG(OodleHandlerComponentLog, Log, TEXT("Oodle Compressed: UnCompressed: %i Compressed: %i Total: %i"),
+					//			UncompressedLength, CompressedLength, BytesSaved);
 
-				INC_DWORD_STAT_BY(STAT_Oodle_BytesRaw, UncompressedLength);
-				INC_DWORD_STAT_BY(STAT_Oodle_BytesCompressed, CompressedLength);
-			}
-			else
-			{
-				Packet.SerializeIntPacked(UncompressedLength);
-			}
+					INC_DWORD_STAT_BY(STAT_Oodle_BytesRaw, UncompressedLength);
+					INC_DWORD_STAT_BY(STAT_Oodle_BytesCompressed, CompressedLength);
+				}
+				else
+				{
+					Packet.SerializeIntPacked(UncompressedLength);
+				}
 
-			break;
+				break;
+			}
 		}
 	}
 }
diff --git a/Engine/Plugins/Runtime/PacketHandlers/CompressionComponents/Oodle/Source/OodleHandlerComponent/Public/OodleHandlerComponent.h b/Engine/Plugins/Runtime/PacketHandlers/CompressionComponents/Oodle/Source/OodleHandlerComponent/Public/OodleHandlerComponent.h
--- a/Engine/Plugins/Runtime/PacketHandlers/CompressionComponents/Oodle/Source/OodleHandlerComponent/Public/OodleHandlerComponent.h
+++ b/Engine/Plugins/Runtime/PacketHandlers/CompressionComponents/Oodle/Source/OodleHandlerComponent/Public/OodleHandlerComponent.h
@@ -91,10 +91,36 @@ public:
 
 
 	/**
-	 * @todo #JohnB
+	 * Initializes FOodleDictionary data, from the specified dictionary file
+	 *
+	 * @param FilePath		The dictionary file path
+	 * @param OutDictionary	The FOodleDictionary struct to write to
 	 */
 	void InitializeDictionary(FString FilePath, FOodleDictionary& OutDictionary);
 
+	/**
+	 * Resolves and returns the default dictionary file paths.
+	 *
+	 * @param OutServerDictionary	The server dictionary path
+	 * @param OutClientDictionary	The client dictionary path
+	 * @param bFailFatal			Whether or not failure to set the dictionary paths, should be fatal
+	 * @return						Whether or not the dictionary paths were successfully set
+	 */
+	bool GetDictionaryPaths(FString& OutServerDictionary, FString& OutClientDictionary, bool bFailFatal=true);
+
+#if !UE_BUILD_SHIPPING
+	/**
+	 * Searches the game directory for alternate/fallback dictionary files, using the *.udic file extension.
+	 * NOTE: This is non-shipping-only, as release games MUST have well-determined dictionary files (for net-binary-compatibility)
+	 *
+	 * @param OutServerDictionary	The server dictionary path
+	 * @param OutClientDictionary	The client dictionary path
+	 * @param bTestOnly				Whether this is being used to test for presence of alternate dictionaries (disables some logging)
+	 * @return						Whether or not alternate dictionaries were found
+	 */
+	bool FindFallbackDictionaries(FString& OutServerDictionary, FString& OutClientDictionary, bool bTestOnly=false);
+#endif
+
 	/* Initializes the handler component */
 	virtual void Initialize() override;
 
@@ -126,6 +152,14 @@ protected:
 #endif
 
 
+	/** Whether or not Oodle is enabled */
+	bool bEnableOodle;
+
+#if !UE_BUILD_SHIPPING
+	/** Search for dictionary files and use them if present - switching mode to Release in process - don't use in shipping */
+	bool bUseDictionaryIfPresent;
+#endif
+
 	/* 
 	* Modes of the component
 	* Specify this in .ini by
diff --git a/Engine/Shaders/ParticleGPUSpriteVertexFactory.usf b/Engine/Shaders/ParticleGPUSpriteVertexFactory.usf
--- a/Engine/Shaders/ParticleGPUSpriteVertexFactory.usf
+++ b/Engine/Shaders/ParticleGPUSpriteVertexFactory.usf
@@ -300,6 +300,11 @@ FVertexFactoryIntermediates GetVertexFactoryIntermediates( FVertexFactoryInput I
 	Intermediates.ParticleWorldPosition = ParticleWorldPosition;
 	Intermediates.VertexWorldPosition = VertexWorldPosition;
 
+#if USE_SPHERICAL_NORMALS
+	// Recenter the particle position for GENERATE_SPHERICAL_PARTICLE_NORMALS to use the right center.
+	Intermediates.ParticleWorldPosition += Size.x * (EmitterUniforms.PivotOffset.x + .5) * TangentRight + Size.y * (EmitterUniforms.PivotOffset.y + .5) * TangentUp;
+#endif
+
 #if USE_PARTICLE_LIGHTING_OFFSET
 	// Hash function based on the particle ID to generate a uniformly distributed 3d offset
 	float3 RandomParticleOffset = frac((ParticleIndex.x + 10) * (ParticleIndex.y + 10) * float3(1341.456345, 2633.578, 5623.983)) * 2 - 1;
diff --git a/Engine/Shaders/ParticleSpriteVertexFactory.usf b/Engine/Shaders/ParticleSpriteVertexFactory.usf
--- a/Engine/Shaders/ParticleSpriteVertexFactory.usf
+++ b/Engine/Shaders/ParticleSpriteVertexFactory.usf
@@ -432,6 +432,11 @@ FVertexFactoryIntermediates GetVertexFactoryIntermediates(FVertexFactoryInput In
 #if USE_PARTICLE_POSITION
 	float ParticleRadius = .5f * min(Input.SizeRotSubImage.x, Input.SizeRotSubImage.y);
 	Intermediates.PositionAndSize = float4(ParticleWorldPosition.xyz, ParticleRadius);
+
+#if GENERATE_SPHERICAL_PARTICLE_NORMALS
+	// Recenter the particle position for GENERATE_SPHERICAL_PARTICLE_NORMALS to use the right center.
+	Intermediates.PositionAndSize.xyz += Size.x * (SpriteVF.PivotOffset.x + .5) * Right + Size.y * (SpriteVF.PivotOffset.y + .5) * Up;
+#endif
 #endif
 
 #if USE_PARTICLE_VELOCITY
diff --git a/Engine/Shaders/ShaderVersion.usf b/Engine/Shaders/ShaderVersion.usf
--- a/Engine/Shaders/ShaderVersion.usf
+++ b/Engine/Shaders/ShaderVersion.usf
@@ -2,4 +2,4 @@
 // Each time the console command is executed it generates a new GUID. As this file is included
 // in common.usf (which should be included in any shader) it allows to invalidate the shader DDC.
 // 
-// GUID = 3062cd7a3314413686da187d920fd485
\ No newline at end of file
+// GUID = 3022cb7a3314413686da187d920fd485
\ No newline at end of file
diff --git a/Engine/Shaders/SlateElementPixelShader.usf b/Engine/Shaders/SlateElementPixelShader.usf
--- a/Engine/Shaders/SlateElementPixelShader.usf
+++ b/Engine/Shaders/SlateElementPixelShader.usf
@@ -305,10 +305,16 @@ float4 Main( VertexToPixelInterpolants VIn ) : SV_Target0
 	if (any(ClipTest < 0))
 	{
 		// Shows red where any pixels discarded because part of the element was clipped
-		OutColor *= float4(1,0,0,1) ;
+		OutColor *= float4(1,0,0,1);
 	}
 #endif
 
+
+#if !PRE_MULTIPLIED_ALPHA
+	OutColor.rgb = OutColor.rgb * OutColor.a;
+#endif
+
+
 	return OutColor;
 }
 
diff --git a/Engine/Source/Developer/AutomationController/Private/AutomationControllerManger.cpp b/Engine/Source/Developer/AutomationController/Private/AutomationControllerManger.cpp
--- a/Engine/Source/Developer/AutomationController/Private/AutomationControllerManger.cpp
+++ b/Engine/Source/Developer/AutomationController/Private/AutomationControllerManger.cpp
@@ -481,7 +481,7 @@ void FAutomationControllerManager::AddPingResult( const FMessageAddress& Respond
 void FAutomationControllerManager::UpdateTests( )
 {
 	static const float CheckTestInterval = 1.0f;
-	static const float GameInstanceLostTimer = 50.0f;
+	static const float GameInstanceLostTimer = 200.0f;
 
 	CheckTestTimer += FPlatformTime::Seconds() - LastTimeUpdateTicked;
 	LastTimeUpdateTicked = FPlatformTime::Seconds();
diff --git a/Engine/Source/Developer/FriendsAndChat/Private/FriendsAndChatStyle.cpp b/Engine/Source/Developer/FriendsAndChat/Private/FriendsAndChatStyle.cpp
--- a/Engine/Source/Developer/FriendsAndChat/Private/FriendsAndChatStyle.cpp
+++ b/Engine/Source/Developer/FriendsAndChat/Private/FriendsAndChatStyle.cpp
@@ -22,12 +22,6 @@ FFriendsAndChatStyle& FFriendsAndChatStyle::SetLargeFriendsFontStyle(const FFrie
 	return *this;
 }
 
-FFriendsAndChatStyle& FFriendsAndChatStyle::SetFriendActionFontStyle(const FFriendsFontStyle& FontStyle)
-{
-	FriendActionFontStyle = FontStyle;
-	return *this;
-}
-
 FFriendsAndChatStyle& FFriendsAndChatStyle::SetFriendsListStyle(const FFriendsListStyle& InFriendsListStyle)
 {
 	FriendsListStyle = InFriendsListStyle;
diff --git a/Engine/Source/Developer/FriendsAndChat/Private/FriendsListStyle.cpp b/Engine/Source/Developer/FriendsAndChat/Private/FriendsListStyle.cpp
--- a/Engine/Source/Developer/FriendsAndChat/Private/FriendsListStyle.cpp
+++ b/Engine/Source/Developer/FriendsAndChat/Private/FriendsListStyle.cpp
@@ -175,12 +175,6 @@ FFriendsListStyle& FFriendsListStyle::SetFriendsListBackground(const FSlateBrush
 	return *this;
 }
 
-FFriendsListStyle& FFriendsListStyle::SetFriendListActionFontColor(const FLinearColor& Color)
-{
-	FriendListActionFontColor = Color;
-	return *this;
-}
-
 FFriendsListStyle& FFriendsListStyle::SetAddFriendEditableTextStyle(const FEditableTextBoxStyle& TextStyle)
 {
 	AddFriendEditableTextStyle = TextStyle;
@@ -205,18 +199,6 @@ FFriendsListStyle& FFriendsListStyle::SetUserHeaderPadding(const FMargin& Margin
 	return *this;
 }
 
-FFriendsListStyle& FFriendsListStyle::SetFriendsListWidth(const float InWidth)
-{
-	FriendsListWidth = InWidth;
-	return *this;
-}
-
-FFriendsListStyle& FFriendsListStyle::SetHasUserHeader(bool InHasUserHeader)
-{
-	HasUserHeader = InHasUserHeader;
-	return *this;
-}
-
 FFriendsListStyle& FFriendsListStyle::SetBackBrush(const FSlateBrush& Brush)
 {
 	BackBrush = Brush;
diff --git a/Engine/Source/Developer/FriendsAndChat/Private/FriendsMarkupStyle.cpp b/Engine/Source/Developer/FriendsAndChat/Private/FriendsMarkupStyle.cpp
--- a/Engine/Source/Developer/FriendsAndChat/Private/FriendsMarkupStyle.cpp
+++ b/Engine/Source/Developer/FriendsAndChat/Private/FriendsMarkupStyle.cpp
@@ -21,3 +21,10 @@ FFriendsMarkupStyle& FFriendsMarkupStyle::SetMarkupTextStyle(const FTextBlockSty
 	MarkupTextStyle = InTextStyle;
 	return *this;
 }
+
+FFriendsMarkupStyle& FFriendsMarkupStyle::SetTipBackgroundBrush(const FSlateBrush& Value)
+{
+	MarkupBackground = Value;
+	return *this;
+}
+
diff --git a/Engine/Source/Developer/FriendsAndChat/Public/FriendsAndChatStyle.h b/Engine/Source/Developer/FriendsAndChat/Public/FriendsAndChatStyle.h
--- a/Engine/Source/Developer/FriendsAndChat/Public/FriendsAndChatStyle.h
+++ b/Engine/Source/Developer/FriendsAndChat/Public/FriendsAndChatStyle.h
@@ -73,10 +73,6 @@ struct FRIENDSANDCHAT_API FFriendsAndChatStyle
 	FFriendsFontStyle FriendsLargeFontStyle;
 	FFriendsAndChatStyle& SetLargeFriendsFontStyle(const FFriendsFontStyle& FontStyle);
 
-	UPROPERTY(EditAnywhere, Category = Appearance)
-	FFriendsFontStyle FriendActionFontStyle;
-	FFriendsAndChatStyle& SetFriendActionFontStyle(const FFriendsFontStyle& FontStyle);
-
 	UPROPERTY(EditAnywhere, Category = Appearance)
 	FCheckBoxStyle CheckBoxStyle;
 	FFriendsAndChatStyle& SetCheckBoxStyle(const FCheckBoxStyle& InCheckBoxStyle);
diff --git a/Engine/Source/Developer/FriendsAndChat/Public/FriendsListStyle.h b/Engine/Source/Developer/FriendsAndChat/Public/FriendsListStyle.h
--- a/Engine/Source/Developer/FriendsAndChat/Public/FriendsListStyle.h
+++ b/Engine/Source/Developer/FriendsAndChat/Public/FriendsListStyle.h
@@ -159,10 +159,6 @@ struct FRIENDSANDCHAT_API FFriendsListStyle
 	FSlateBrush FriendsListBackground;
 	FFriendsListStyle& SetFriendsListBackground(const FSlateBrush& InBrush);
 
-	UPROPERTY(EditAnywhere, Category = Appearance)
-	FLinearColor FriendListActionFontColor;
-	FFriendsListStyle& SetFriendListActionFontColor(const FLinearColor& InColor);
-
 	UPROPERTY(EditAnywhere, Category = Appearance)
 	FEditableTextBoxStyle AddFriendEditableTextStyle;
 	FFriendsListStyle& SetAddFriendEditableTextStyle(const FEditableTextBoxStyle& InEditableTextStyle);
@@ -179,14 +175,6 @@ struct FRIENDSANDCHAT_API FFriendsListStyle
 	FMargin UserHeaderPadding;
 	FFriendsListStyle& SetUserHeaderPadding(const FMargin& InUserHeaderPadding);
 
-	UPROPERTY(EditAnywhere, Category = Appearance)
-	float FriendsListWidth;
-	FFriendsListStyle& SetFriendsListWidth(const float FriendsListLength);
-
-	UPROPERTY(EditAnywhere, Category = Appearance)
-	bool HasUserHeader;
-	FFriendsListStyle& SetHasUserHeader(bool InHasUserHeader);
-
 	UPROPERTY(EditAnywhere, Category = Appearance)
 	FSlateBrush BackBrush;
 	FFriendsListStyle& SetBackBrush(const FSlateBrush& Brush);
@@ -203,6 +191,48 @@ struct FRIENDSANDCHAT_API FFriendsListStyle
 	FSlateBrush SettingsBrush;
 	FFriendsListStyle& SetSettingsBrush(const FSlateBrush& Brush);
 
+	UPROPERTY(EditAnywhere, Category = Appearance)
+	FMargin BackButtonMargin;
+
+	UPROPERTY(EditAnywhere, Category = Appearance)
+	FMargin HeaderButtonMargin;
+
+	UPROPERTY(EditAnywhere, Category = Appearance)
+	FMargin FriendsListMargin;
+
+	UPROPERTY(EditAnywhere, Category = Appearance)
+	FMargin BackButtonContentMargin;
+
+	UPROPERTY(EditAnywhere, Category = Appearance)
+	FMargin FriendsListNoFriendsMargin;
+
+	UPROPERTY(EditAnywhere, Category = Appearance)
+	FMargin FriendsListHeaderMargin;
+
+	UPROPERTY(EditAnywhere, Category = Appearance)
+	FMargin FriendsListHeaderCountMargin;
+
+	UPROPERTY(EditAnywhere, Category = Appearance)
+	FMargin HeaderButtonContentMargin;
+
+	UPROPERTY(EditAnywhere, Category = Appearance)
+	FMargin FriendItemMargin;
+
+	UPROPERTY(EditAnywhere, Category = Appearance)
+	FMargin FriendItemStatusMargin;
+
+	UPROPERTY(EditAnywhere, Category = Appearance)
+	FMargin FriendItemPresenceMargin;
+
+	UPROPERTY(EditAnywhere, Category = Appearance)
+	FMargin ConfirmationBorderMargin;
+
+	UPROPERTY(EditAnywhere, Category = Appearance)
+	FMargin ConfirmationButtonMargin;
+
+	UPROPERTY(EditAnywhere, Category = Appearance)
+	FMargin ConfirmationButtonContentMargin;
+
 
 // Clan Settings
 
diff --git a/Engine/Source/Developer/FriendsAndChat/Public/FriendsMarkupStyle.h b/Engine/Source/Developer/FriendsAndChat/Public/FriendsMarkupStyle.h
--- a/Engine/Source/Developer/FriendsAndChat/Public/FriendsMarkupStyle.h
+++ b/Engine/Source/Developer/FriendsAndChat/Public/FriendsMarkupStyle.h
@@ -48,5 +48,8 @@ struct FRIENDSANDCHAT_API FFriendsMarkupStyle
 	FTextBlockStyle MarkupTextStyle;
 	FFriendsMarkupStyle& SetMarkupTextStyle(const FTextBlockStyle& InTextStyle);
 
+	UPROPERTY(EditAnywhere, Category = Appearance)
+	FSlateBrush MarkupBackground;
+	FFriendsMarkupStyle& SetTipBackgroundBrush(const FSlateBrush& InChatBackgroundBrush);
 };
 
diff --git a/Engine/Source/Developer/FunctionalTesting/Classes/FunctionalTest.h b/Engine/Source/Developer/FunctionalTesting/Classes/FunctionalTest.h
--- a/Engine/Source/Developer/FunctionalTesting/Classes/FunctionalTest.h
+++ b/Engine/Source/Developer/FunctionalTesting/Classes/FunctionalTest.h
@@ -7,6 +7,14 @@
 
 class UBillboardComponent;
 
+//Experimental effort at automated cpu captures from the functional testing.
+class FFunctionalTestExternalProfiler : public FScopedExternalProfilerBase
+{
+public:
+	void StartProfiler(const bool bWantPause){ StartScopedTimer(bWantPause); }
+	void StopProfiler(){StopScopedTimer();}
+}; 
+
 // Used to measure a distribution
 struct FStatisticalFloat
 {
@@ -61,41 +69,77 @@ private:
 	int32 NumSamples;
 };
 
-/** A set of simple perf stats recorded over a period of frames. */
-struct FUNCTIONALTESTING_API FPerfStatsRecord
+struct FStatsData
 {
-	FPerfStatsRecord(FString InName);
+	FStatsData() :NumFrames(0), SumTimeSeconds(0.0f){}
 
-	FString Name;
 	uint32 NumFrames;
 	uint32 SumTimeSeconds;
 	FStatisticalFloat FrameTimeTracker;
 	FStatisticalFloat GameThreadTimeTracker;
 	FStatisticalFloat RenderThreadTimeTracker;
 	FStatisticalFloat GPUTimeTracker;
+};
+
+/** A set of simple perf stats recorded over a period of frames. */
+struct FUNCTIONALTESTING_API FPerfStatsRecord
+{
+	FPerfStatsRecord(FString InName);
+
+	FString Name;
+
+	/** Stats data for the period we're interested in timing. */
+	FStatsData Record;
+	/** Stats data for the baseline. */
+	FStatsData Baseline;
 
-	void Sample(AActor* Owner, float DeltaSeconds);
+	void Sample(UWorld* Owner, float DeltaSeconds, bool bBaseline);
 
 	FString GetReportString()const;
+	FString GetBaselineString()const;
+	FString GetRecordString()const;
+
+	void GetGPUTimes(double& OutMin, double& OutMax, double& OutAvg)const;
+	void GetGameThreadTimes(double& OutMin, double& OutMax, double& OutAvg)const;
+	void GetRenderThreadTimes(double& OutMin, double& OutMax, double& OutAvg)const;
 };
 
-/** Simple class to record some basic performance stats during functional tests and write out the results. */
+
+/** 
+ * Class for use with functional tests which provides various performance measuring features. 
+ * Recording of basic, unintrusive performance stats.
+ * Automatic captures using external CPU and GPU profilers.
+ * Triggering and ending of writing full stats to a file.
+*/
 UCLASS(Blueprintable)
-class FUNCTIONALTESTING_API UPerfStatsRecorder : public UObject
+class FUNCTIONALTESTING_API UAutomationPerformaceHelper : public UObject
 {
 	GENERATED_BODY()
 
-	bool bRecording;
 	TArray<FPerfStatsRecord> Records;
-
+	bool bRecordingBasicStats;
+	bool bRecordingBaselineBasicStats;
+	bool bRecordingCPUCapture;
+	bool bRecordingStatsFile;
 public:
 
-	UPerfStatsRecorder();
+	UAutomationPerformaceHelper();
+
+	//Begin basic stat recording
+
+	UFUNCTION(BlueprintCallable, Category = Perf)
+	void Tick(float DeltaSeconds);
 
 	/** Adds a sample to the stats counters for the current performance stats record. */
 	UFUNCTION(BlueprintCallable, Category = Perf)
-	void Sample(AActor* Owner, float DeltaSeconds);
-	/** Begins recording a new named performance stats record. */
+	void Sample(float DeltaSeconds);
+	/** Begins recording a new named performance stats record. We start by recording the baseline */
+	UFUNCTION(BlueprintCallable, Category = Perf)
+	void BeginRecordingBaseline(FString RecordName);
+	/** Stops recording the baseline and moves to the main record. */
+	UFUNCTION(BlueprintCallable, Category = Perf)
+	void EndRecordingBaseline();
+	/** Begins recording a new named performance stats record. We start by recording the baseline. */
 	UFUNCTION(BlueprintCallable, Category = Perf)
 	void BeginRecording(FString RecordName);
 	/** Stops recording performance stats. */
@@ -106,9 +150,49 @@ public:
 	void WriteLogFile(const FString& CaptureDir, const FString& CaptureExtension);
 	/** Returns true if this stats tracker is currently recording performance stats. */
 	UFUNCTION(BlueprintCallable, Category = Perf)
-	FORCEINLINE bool IsRecording()const { return bRecording; }
+	bool IsRecording()const;
+
+	/** Does any init work across all tests.. */
+	UFUNCTION(BlueprintCallable, Category = Perf)
+	void OnBeginTests();
+	/** Does any final work needed as all tests are complete. */
+	UFUNCTION(BlueprintCallable, Category = Perf)
+	void OnAllTestsComplete();
 
 	const FPerfStatsRecord* GetCurrentRecord()const;
+
+	UFUNCTION(BlueprintCallable, Category = Perf)
+	bool IsCurrentRecordWithinGPUBudget(float InMaxTimeBudget, float InAvgTimeBudget = -1.0f)const;
+	UFUNCTION(BlueprintCallable, Category = Perf)
+	bool IsCurrentRecordWithinGameThreadBudget(float InMaxTimeBudget, float InAvgTimeBudget = -1.0f)const;
+	UFUNCTION(BlueprintCallable, Category = Perf)
+	bool IsCurrentRecordWithinRenderThreadBudget(float InMaxTimeBudget, float InAvgTimeBudget = -1.0f)const;
+	//End basic stats recording.
+
+	// Automatic traces capturing 
+
+	/** Communicates with external profiler to being a CPU capture. */
+	UFUNCTION(BlueprintCallable, Category = Perf)
+	void StartCPUProfiling();
+	/** Communicates with external profiler to end a CPU capture. */
+	UFUNCTION(BlueprintCallable, Category = Perf)
+	void StopCPUProfiling();
+	/** Communicates with an external GPU profiler to trigger a GPU trace. */
+	UFUNCTION(BlueprintCallable, Category = Perf)
+	void TriggerGPUTrace();
+	/** Begins recording stats to a file. */
+	UFUNCTION(BlueprintCallable, Category = Perf)
+	void BeginStatsFile(const FString& RecordName);
+	/** Ends recording stats to a file. */
+	UFUNCTION(BlueprintCallable, Category = Perf)
+	void EndStatsFile();
+
+	FFunctionalTestExternalProfiler ExternalProfiler;
+
+	/** The path and base name for all output files. */
+	FString OutputFileBase;
+	
+	FString StartOfTestingTime;
 };
 
 UENUM()
@@ -184,6 +268,11 @@ public:
 	UPROPERTY(EditAnywhere, BlueprintReadOnly, Category=FunctionalTesting)
 	uint32 bIsEnabled:1;
 
+	/** List of causes we need a re-run. */
+	TArray<FName> RerunCauses;
+	/** Cause of the current rerun if we're in a named rerun. */
+	FName CurrentRerunCause;
+
 public:
 
 	virtual bool StartTest(const TArray<FString>& Params = TArray<FString>());
@@ -209,6 +298,14 @@ public:
 
 	virtual bool WantsToRunAgain() const { return false; }
 
+	/** Causes the test to be rerun for a specific named reason. */
+	UFUNCTION(BlueprintCallable, Category = "FunctionalTesting")
+	void AddRerun(FName Reason);
+
+	/** Returns the current re-run reason if we're in a named re-run. */
+	UFUNCTION(BlueprintCallable, Category = "FunctionalTesting")
+	FName GetCurrentRerunReason()const;
+
 	UFUNCTION(BlueprintImplementableEvent, Category = "FunctionalTesting")
 	FString OnAdditionalTestFinishedMessageRequest(EFunctionalTestResult::Type TestResult) const;
 	
diff --git a/Engine/Source/Developer/FunctionalTesting/Classes/FunctionalTestingManager.h b/Engine/Source/Developer/FunctionalTesting/Classes/FunctionalTestingManager.h
--- a/Engine/Source/Developer/FunctionalTesting/Classes/FunctionalTestingManager.h
+++ b/Engine/Source/Developer/FunctionalTesting/Classes/FunctionalTestingManager.h
@@ -27,6 +27,9 @@ class UFunctionalTestingManager : public UBlueprintFunctionLibrary
 
 	UPROPERTY(BlueprintAssignable)
 	FFunctionalTestEventSignature OnSetupTests;
+
+	UPROPERTY(BlueprintAssignable)
+	FFunctionalTestEventSignature OnTestsComplete;
 	
 	UFUNCTION(BlueprintCallable, Category="FunctionalTesting", meta=(WorldContext="WorldContext", CallableWithoutWorldContext ) )
 	/** Triggers in sequence all functional tests found on the level.
@@ -77,3 +80,14 @@ protected:
 private:
 	FTimerHandle TriggerFirstValidTestTimerHandle;
 };
+
+UCLASS(abstract, Blueprintable, MinimalAPI)
+class APhasedAutomationActorBase : public AActor
+{
+	GENERATED_BODY()
+
+public:
+
+	UFUNCTION(BlueprintImplementableEvent, Category = "Automation")
+	void OnFunctionalTestingComplete();
+};
\ No newline at end of file
diff --git a/Engine/Source/Developer/FunctionalTesting/Private/FunctionalTest.cpp b/Engine/Source/Developer/FunctionalTesting/Private/FunctionalTest.cpp
--- a/Engine/Source/Developer/FunctionalTesting/Private/FunctionalTest.cpp
+++ b/Engine/Source/Developer/FunctionalTesting/Private/FunctionalTest.cpp
@@ -78,6 +78,9 @@ void AFunctionalTest::Tick(float DeltaSeconds)
 		return;
 	}
 
+	//Do not collect garbage during the test. We force GC at the end.
+	GetWorld()->DelayGarbageCollection();
+
 	TotalTime += DeltaSeconds;
 	if (TimeLimit > 0.f && TotalTime > TimeLimit)
 	{
@@ -93,6 +96,9 @@ bool AFunctionalTest::StartTest(const TArray<FString>& Params)
 {
 	FailureMessage = TEXT("");
 	
+	//Do not collect garbage during the test. We force GC at the end.
+	GetWorld()->DelayGarbageCollection();
+
 	TotalTime = 0.f;
 	if (TimeLimit > 0)
 	{
@@ -118,6 +124,9 @@ void AFunctionalTest::FinishTest(TEnumAsByte<EFunctionalTestResult::Type> TestRe
 		return;
 	}
 
+	//Force GC at the end of every test.
+	GetWorld()->ForceGarbageCollection();
+
 	Result = TestResult;
 
 	bIsRunning = false;
@@ -235,6 +244,16 @@ void AFunctionalTest::GatherRelevantActors(TArray<AActor*>& OutActors) const
 	OutActors.Append(DebugGatherRelevantActors());
 }
 
+void AFunctionalTest::AddRerun(FName Reason)
+{
+	RerunCauses.Add(Reason);
+}
+
+FName AFunctionalTest::GetCurrentRerunReason()const
+{
+	return CurrentRerunCause;
+}
+
 void AFunctionalTest::RegisterAutoDestroyActor(AActor* ActorToAutoDestroy)
 {
 	AutoDestroyActors.AddUnique(ActorToAutoDestroy);
@@ -311,77 +330,178 @@ UBillboardComponent* AFunctionalTest::GetSpriteComponent() { return SpriteCompon
 
 FPerfStatsRecord::FPerfStatsRecord(FString InName)
 : Name(InName)
-, NumFrames(0.0f)
-, SumTimeSeconds(0.0f)
 {
 }
 
 FString FPerfStatsRecord::GetReportString() const
 {
-	return FString::Printf(TEXT("%s,%d,%.2f,%.2f,%.2f,%.2f,%.2f,%.2f,%.2f,%.2f,%.2f,%.2f,%.2f,%.2f,%.2f"),
+	return FString::Printf(TEXT("%s,%.2f,%.2f,%.2f,%.2f,%.2f,%.2f,%.2f,%.2f,%.2f,%.2f,%.2f,%.2f"),
 		*Name,
-		NumFrames,
-		SumTimeSeconds * 1000.0f,
-		FrameTimeTracker.GetMinValue(),
-		FrameTimeTracker.GetAvgValue(),
-		FrameTimeTracker.GetMaxValue(),
-		RenderThreadTimeTracker.GetMinValue(),
-		RenderThreadTimeTracker.GetAvgValue(),
-		RenderThreadTimeTracker.GetMaxValue(),
-		GameThreadTimeTracker.GetMinValue(),
-		GameThreadTimeTracker.GetAvgValue(),
-		GameThreadTimeTracker.GetMaxValue(),
-		GPUTimeTracker.GetMinValue(),
-		GPUTimeTracker.GetAvgValue(),
-		GPUTimeTracker.GetMaxValue());
-}
-
-void FPerfStatsRecord::Sample(AActor* Owner, float DeltaSeconds)
-{
-	check(Owner);
-
-	const FStatUnitData* StatUnitData = Owner->GetWorld()->GetGameViewport()->GetStatUnitData();
-	check(StatUnitData);
+		Record.FrameTimeTracker.GetMinValue() - Baseline.FrameTimeTracker.GetMinValue(),
+		Record.FrameTimeTracker.GetAvgValue() - Baseline.FrameTimeTracker.GetAvgValue(),
+		Record.FrameTimeTracker.GetMaxValue() - Baseline.FrameTimeTracker.GetMaxValue(),
+		Record.RenderThreadTimeTracker.GetMinValue() - Baseline.RenderThreadTimeTracker.GetMinValue(),
+		Record.RenderThreadTimeTracker.GetAvgValue() - Baseline.RenderThreadTimeTracker.GetAvgValue(),
+		Record.RenderThreadTimeTracker.GetMaxValue() - Baseline.RenderThreadTimeTracker.GetMaxValue(),
+		Record.GameThreadTimeTracker.GetMinValue() - Baseline.GameThreadTimeTracker.GetMinValue(),
+		Record.GameThreadTimeTracker.GetAvgValue() - Baseline.GameThreadTimeTracker.GetAvgValue(),
+		Record.GameThreadTimeTracker.GetMaxValue() - Baseline.GameThreadTimeTracker.GetMaxValue(),
+		Record.GPUTimeTracker.GetMinValue() - Baseline.GPUTimeTracker.GetMinValue(),
+		Record.GPUTimeTracker.GetAvgValue() - Baseline.GPUTimeTracker.GetAvgValue(),
+		Record.GPUTimeTracker.GetMaxValue() - Baseline.GPUTimeTracker.GetMaxValue());
+}
+
+FString FPerfStatsRecord::GetBaselineString() const
+{
+	return FString::Printf(TEXT("%s,%.2f,%.2f,%.2f,%.2f,%.2f,%.2f,%.2f,%.2f,%.2f,%.2f,%.2f,%.2f"),
+		*Name,
+		Baseline.FrameTimeTracker.GetMinValue(),
+		Baseline.FrameTimeTracker.GetAvgValue(),
+		Baseline.FrameTimeTracker.GetMaxValue(),
+		Baseline.RenderThreadTimeTracker.GetMinValue(),
+		Baseline.RenderThreadTimeTracker.GetAvgValue(),
+		Baseline.RenderThreadTimeTracker.GetMaxValue(),
+		Baseline.GameThreadTimeTracker.GetMinValue(),
+		Baseline.GameThreadTimeTracker.GetAvgValue(),
+		Baseline.GameThreadTimeTracker.GetMaxValue(),
+		Baseline.GPUTimeTracker.GetMinValue(),
+		Baseline.GPUTimeTracker.GetAvgValue(),
+		Baseline.GPUTimeTracker.GetMaxValue());
+}
+
+FString FPerfStatsRecord::GetRecordString() const
+{
+	return FString::Printf(TEXT("%s,%.2f,%.2f,%.2f,%.2f,%.2f,%.2f,%.2f,%.2f,%.2f,%.2f,%.2f,%.2f"),
+		*Name,
+		Record.FrameTimeTracker.GetMinValue(),
+		Record.FrameTimeTracker.GetAvgValue(),
+		Record.FrameTimeTracker.GetMaxValue(),
+		Record.RenderThreadTimeTracker.GetMinValue(),
+		Record.RenderThreadTimeTracker.GetAvgValue(),
+		Record.RenderThreadTimeTracker.GetMaxValue(),
+		Record.GameThreadTimeTracker.GetMinValue(),
+		Record.GameThreadTimeTracker.GetAvgValue(),
+		Record.GameThreadTimeTracker.GetMaxValue(),
+		Record.GPUTimeTracker.GetMinValue(),
+		Record.GPUTimeTracker.GetAvgValue(),
+		Record.GPUTimeTracker.GetMaxValue());
+}
+
+void FPerfStatsRecord::GetGPUTimes(double& OutMin, double& OutMax, double& OutAvg)const
+{
+	OutMin = Record.GPUTimeTracker.GetMinValue() - Baseline.GPUTimeTracker.GetMinValue();
+	OutMax = Record.GPUTimeTracker.GetMaxValue() - Baseline.GPUTimeTracker.GetMaxValue();
+	OutAvg = Record.GPUTimeTracker.GetAvgValue() - Baseline.GPUTimeTracker.GetAvgValue();
+}
 
-	FrameTimeTracker.AddSample(StatUnitData->RawFrameTime);
-	GameThreadTimeTracker.AddSample(FPlatformTime::ToMilliseconds(GGameThreadTime));
-	RenderThreadTimeTracker.AddSample(FPlatformTime::ToMilliseconds(GRenderThreadTime));
-	GPUTimeTracker.AddSample(FPlatformTime::ToMilliseconds(GGPUFrameTime));
-	NumFrames++;
-	SumTimeSeconds += DeltaSeconds;
+void FPerfStatsRecord::GetGameThreadTimes(double& OutMin, double& OutMax, double& OutAvg)const
+{
+	OutMin = Record.GameThreadTimeTracker.GetMinValue() - Baseline.GameThreadTimeTracker.GetMinValue();
+	OutMax = Record.GameThreadTimeTracker.GetMaxValue() - Baseline.GameThreadTimeTracker.GetMaxValue();
+	OutAvg = Record.GameThreadTimeTracker.GetAvgValue() - Baseline.GameThreadTimeTracker.GetAvgValue();
 }
 
-UPerfStatsRecorder::UPerfStatsRecorder()
-: bRecording(false)
+void FPerfStatsRecord::GetRenderThreadTimes(double& OutMin, double& OutMax, double& OutAvg)const
 {
+	OutMin = Record.RenderThreadTimeTracker.GetMinValue() - Baseline.RenderThreadTimeTracker.GetMinValue();
+	OutMax = Record.RenderThreadTimeTracker.GetMaxValue() - Baseline.RenderThreadTimeTracker.GetMaxValue();
+	OutAvg = Record.RenderThreadTimeTracker.GetAvgValue() - Baseline.RenderThreadTimeTracker.GetAvgValue();
+}
+
+void FPerfStatsRecord::Sample(UWorld* World, float DeltaSeconds, bool bBaseline)
+{
+	check(World);
+
+	const FStatUnitData* StatUnitData = World->GetGameViewport()->GetStatUnitData();
+	check(StatUnitData);
 
+	if (bBaseline)
+	{
+		Baseline.FrameTimeTracker.AddSample(StatUnitData->RawFrameTime);
+		Baseline.GameThreadTimeTracker.AddSample(FPlatformTime::ToMilliseconds(GGameThreadTime));
+		Baseline.RenderThreadTimeTracker.AddSample(FPlatformTime::ToMilliseconds(GRenderThreadTime));
+		Baseline.GPUTimeTracker.AddSample(FPlatformTime::ToMilliseconds(GGPUFrameTime));
+		Baseline.NumFrames++;
+		Baseline.SumTimeSeconds += DeltaSeconds;
+	}
+	else
+	{
+		Record.FrameTimeTracker.AddSample(StatUnitData->RawFrameTime);
+		Record.GameThreadTimeTracker.AddSample(FPlatformTime::ToMilliseconds(GGameThreadTime));
+		Record.RenderThreadTimeTracker.AddSample(FPlatformTime::ToMilliseconds(GRenderThreadTime));
+		Record.GPUTimeTracker.AddSample(FPlatformTime::ToMilliseconds(GGPUFrameTime));
+		Record.NumFrames++;
+		Record.SumTimeSeconds += DeltaSeconds;
+	}
+}
+
+UAutomationPerformaceHelper::UAutomationPerformaceHelper()
+: bRecordingBasicStats(false)
+, bRecordingBaselineBasicStats(false)
+, bRecordingCPUCapture(false)
+, bRecordingStatsFile(false)
+{
 }
 
-void UPerfStatsRecorder::BeginRecording(FString RecordName)
+void UAutomationPerformaceHelper::BeginRecordingBaseline(FString RecordName)
 {
-	bRecording = true;
+	bRecordingBasicStats = true;
+	bRecordingBaselineBasicStats = true;
 	Records.Add(FPerfStatsRecord(RecordName));
+	GEngine->SetEngineStat(GetOuter()->GetWorld(), GetOuter()->GetWorld()->GetGameViewport(), TEXT("Unit"), true);
+}
+
+void UAutomationPerformaceHelper::EndRecordingBaseline()
+{
+	bRecordingBaselineBasicStats = false;
+	bRecordingBasicStats = false;
+}
+
+void UAutomationPerformaceHelper::BeginRecording(FString RecordName)
+{
+	//Ensure we're recording engine stats.
+	GEngine->SetEngineStat(GetOuter()->GetWorld(), GetOuter()->GetWorld()->GetGameViewport(), TEXT("Unit"), true);
+	bRecordingBasicStats = true;
+	bRecordingBaselineBasicStats = false;
+	//We recorded a baseline for this so continue adding real stats to this one.
+	if (GetCurrentRecord()->Name != RecordName)
+	{
+		Records.Add(FPerfStatsRecord(RecordName));
+	}
+	
+	//We should have a previously baselined record or have just added one.
+	check(Records.Num() > 0);
 }
 
-void UPerfStatsRecorder::EndRecording()
+void UAutomationPerformaceHelper::EndRecording()
 {
 	if (const FPerfStatsRecord* Record = GetCurrentRecord())
 	{
 		UE_LOG(LogFunctionalTest, Log, TEXT("Finished Perf Stats Record:\n%s"), *Record->GetReportString());
 	}
-	bRecording = false;
+	bRecordingBasicStats = false;
+}
+
+void UAutomationPerformaceHelper::Tick(float DeltaSeconds)
+{
+	if (bRecordingBasicStats)
+	{
+		Sample(DeltaSeconds);
+	}
+
+	//Other stats need ticking?
 }
 
-void UPerfStatsRecorder::Sample(AActor* Owner, float DeltaSeconds)
+void UAutomationPerformaceHelper::Sample(float DeltaSeconds)
 {
 	int32 Index = Records.Num() - 1;
-	if (Index >= 0 && bRecording)
+	if (Index >= 0 && bRecordingBasicStats)
 	{
-		Records[Index].Sample(Owner, DeltaSeconds);
+		Records[Index].Sample(GetOuter()->GetWorld(), DeltaSeconds, bRecordingBaselineBasicStats);
 	}
 }
 
-void UPerfStatsRecorder::WriteLogFile(const FString& CaptureDir, const FString& CaptureExtension)
+void UAutomationPerformaceHelper::WriteLogFile(const FString& CaptureDir, const FString& CaptureExtension)
 {
 	FString PathName = FPaths::ProfilingDir();
 	if (!CaptureDir.IsEmpty())
@@ -396,31 +516,148 @@ void UPerfStatsRecorder::WriteLogFile(const FString& CaptureDir, const FString&
 		Extension = TEXT("perf.csv");
 	}
 
-	const FString Filename = CreateProfileFilename(Extension, true);
+	const FString Filename = OutputFileBase + Extension;
 	const FString FilenameFull = PathName + Filename;
 
-	const FString LogHeader = TEXT("TestName,NumFrames,Duration,MinFrameTime,AvgFrameTime,MaxFrameTime,MinRT,AvgRT,MaxRT,MinGT,AvgGT,MaxGT,MinGPU,AvgGPU,MaxGPU\n");
+	const FString LogHeader = TEXT("TestName,MinFrameTime,AvgFrameTime,MaxFrameTime,MinRT,AvgRT,MaxRT,MinGT,AvgGT,MaxGT,MinGPU,AvgGPU,MaxGPU\n");
 
-	FString FileContents = LogHeader;
+	FString AdjustedTable;
 	for (FPerfStatsRecord& Record : Records)
 	{
-		FileContents += Record.GetReportString() + FString(TEXT("\n"));
+		AdjustedTable += Record.GetReportString() + FString(TEXT("\n"));
 	}
+	FString RecordTable;
+	for (FPerfStatsRecord& Record : Records)
+	{
+		RecordTable += Record.GetRecordString() + FString(TEXT("\n"));
+	}
+	FString BaselineTable;
+	for (FPerfStatsRecord& Record : Records)
+	{
+		BaselineTable += Record.GetBaselineString() + FString(TEXT("\n"));
+	}
+
+	FString FileContents = FString::Printf(TEXT("Adjusted Results\n%s%s\nRaw Results\n%s%s\nBaseline Results\n%s%s\n"), *LogHeader, *AdjustedTable, *LogHeader, *RecordTable, *LogHeader, *BaselineTable);
 
 	FFileHelper::SaveStringToFile(FileContents, *FilenameFull);
 
 	UE_LOG(LogTemp, Display, TEXT("Finished test, wrote file to %s"), *FilenameFull);
 
 	Records.Empty();
-	bRecording = false;
+	bRecordingBasicStats = false;
+	bRecordingBaselineBasicStats = false;
+}
+
+bool UAutomationPerformaceHelper::IsRecording()const 
+{
+	return bRecordingBasicStats;
+}
+
+void UAutomationPerformaceHelper::OnBeginTests()
+{
+	OutputFileBase = CreateProfileFilename(TEXT(""), true);
+	StartOfTestingTime = FDateTime::Now().ToString();
+}
+
+void UAutomationPerformaceHelper::OnAllTestsComplete()
+{
+	if (bRecordingBaselineBasicStats)
+	{
+		EndRecordingBaseline();
+	}
+
+	if (bRecordingBasicStats)
+	{
+		EndRecording();
+	}
+
+	if (bRecordingCPUCapture)
+	{
+		StopCPUProfiling();
+	}
+
+	if (bRecordingStatsFile)
+	{
+		EndStatsFile();
+	}
+
+	if (Records.Num() > 0)
+	{
+		WriteLogFile(TEXT(""), TEXT("perf.csv"));
+	}
 }
 
-const FPerfStatsRecord* UPerfStatsRecorder::GetCurrentRecord()const
+bool UAutomationPerformaceHelper::IsCurrentRecordWithinGPUBudget(float InMaxTimeBudget, float InAvgTimeBudget)const
+{
+	double Min, Max, Avg;
+	bool bInBudget = true;
+	if (const FPerfStatsRecord* Curr = GetCurrentRecord())
+	{
+		Curr->GetGPUTimes(Min, Max, Avg);
+		bInBudget = Max < InMaxTimeBudget && (InAvgTimeBudget < 0.0f || Avg < InAvgTimeBudget);
+	}
+	return bInBudget;
+}
+
+bool UAutomationPerformaceHelper::IsCurrentRecordWithinGameThreadBudget(float InMaxTimeBudget, float InAvgTimeBudget)const
+{
+	double Min, Max, Avg;
+	bool bInBudget = true;
+	if (const FPerfStatsRecord* Curr = GetCurrentRecord())
+	{
+		Curr->GetGameThreadTimes(Min, Max, Avg);
+		bInBudget = Max < InMaxTimeBudget && (InAvgTimeBudget < 0.0f || Avg < InAvgTimeBudget);
+	}
+	return bInBudget;
+}
+
+bool UAutomationPerformaceHelper::IsCurrentRecordWithinRenderThreadBudget(float InMaxTimeBudget, float InAvgTimeBudget)const
+{
+	double Min, Max, Avg;
+	bool bInBudget = true;
+	if (const FPerfStatsRecord* Curr = GetCurrentRecord())
+	{
+		Curr->GetRenderThreadTimes(Min, Max, Avg);
+		bInBudget = Max < InMaxTimeBudget && (InAvgTimeBudget < 0.0f || Avg < InAvgTimeBudget);
+	}
+	return bInBudget;
+}
+
+const FPerfStatsRecord* UAutomationPerformaceHelper::GetCurrentRecord()const
 {
 	int32 Index = Records.Num() - 1;
-	if (Index >= 0 && bRecording)
+	if (Index >= 0)
 	{
 		return &Records[Index];
 	}
 	return nullptr;
+}
+
+void UAutomationPerformaceHelper::StartCPUProfiling()
+{
+	UE_LOG(LogFunctionalTest, Log, TEXT("START PROFILING..."));
+	ExternalProfiler.StartProfiler(false);
+}
+
+void UAutomationPerformaceHelper::StopCPUProfiling()
+{
+	UE_LOG(LogFunctionalTest, Log, TEXT("STOP PROFILING..."));
+	ExternalProfiler.StopProfiler();
+}
+
+void UAutomationPerformaceHelper::TriggerGPUTrace()
+{
+	// Need to look at Razor GPU to work out what can be done here.
+}
+
+void UAutomationPerformaceHelper::BeginStatsFile(const FString& RecordName)
+{
+	FString MapName = GetOuter()->GetWorld()->GetMapName();
+	FString Cmd = FString::Printf(TEXT("Stat StartFile %s-%s/%s.ue4stats"), *MapName, *StartOfTestingTime, *RecordName);
+	GEngine->Exec(GetOuter()->GetWorld(), *Cmd);
+}
+
+void UAutomationPerformaceHelper::EndStatsFile()
+{
+	GEngine->Exec(GetOuter()->GetWorld(), TEXT("Stat StopFile"));
 }
\ No newline at end of file
diff --git a/Engine/Source/Developer/FunctionalTesting/Private/FunctionalTestingManager.cpp b/Engine/Source/Developer/FunctionalTesting/Private/FunctionalTestingManager.cpp
--- a/Engine/Source/Developer/FunctionalTesting/Private/FunctionalTestingManager.cpp
+++ b/Engine/Source/Developer/FunctionalTesting/Private/FunctionalTestingManager.cpp
@@ -87,6 +87,8 @@ bool UFunctionalTestingManager::RunAllFunctionalTests(UObject* WorldContext, boo
 		UE_LOG(LogFunctionalTest, Warning, TEXT("Functional tests are already running, aborting."));
 		return true;
 	}
+	
+	WorldContext->GetWorld()->ForceGarbageCollection(true);
 
 	Manager->bFinished = false;
 	Manager->bLooped = bRunLooped;
@@ -105,6 +107,12 @@ bool UFunctionalTestingManager::RunAllFunctionalTests(UObject* WorldContext, boo
 	}
 	else
 	{
+		for (TActorIterator<APhasedAutomationActorBase> It(WorldContext->GetWorld()); It; ++It)
+		{
+			APhasedAutomationActorBase* PAA = (*It);
+			Manager->OnTestsComplete.AddDynamic(PAA, &APhasedAutomationActorBase::OnFunctionalTestingComplete);
+		}
+
 		for (TActorIterator<AFunctionalTest> It(WorldContext->GetWorld()); It; ++It)
 		{
 			AFunctionalTest* Test = (*It);
@@ -210,12 +218,21 @@ void UFunctionalTestingManager::NotifyTestDone(AFunctionalTest* FTest)
 
 	if (FTest->OnWantsReRunCheck() == false && FTest->WantsToRunAgain() == false)
 	{
-		TestsLeft.RemoveSingle(FTest);
-		/*if (bDiscardSuccessfulTests && FTest->IsSuccessful())
+		//We can also do named reruns. These are lower priority than those triggered above.
+		//These names can be querried by phases to alter behviour in re-runs.
+		if (FTest->RerunCauses.Num() > 0)
 		{
+			FTest->CurrentRerunCause = FTest->RerunCauses.Pop();
+		}
+		else
+		{
+			TestsLeft.RemoveSingle(FTest);
+			/*if (bDiscardSuccessfulTests && FTest->IsSuccessful())
+			{
 			AllTests.RemoveSingle(FTest);
-		}*/
-		FTest->CleanUp();
+			}*/
+			FTest->CleanUp();
+		}
 	}
 
 	if (TestsLeft.Num() > 0 || TestReproStrings.Num() > 0)
@@ -259,6 +276,7 @@ void UFunctionalTestingManager::AllTestsDone()
 	}
 	else
 	{
+		OnTestsComplete.Broadcast();
 		bFinished = true;
 		UE_LOG(LogFunctionalTest, Log, TEXT("DONE."));
 		RemoveFromRoot();
diff --git a/Engine/Source/Developer/HotReload/Private/HotReload.cpp b/Engine/Source/Developer/HotReload/Private/HotReload.cpp
--- a/Engine/Source/Developer/HotReload/Private/HotReload.cpp
+++ b/Engine/Source/Developer/HotReload/Private/HotReload.cpp
@@ -390,8 +390,74 @@ namespace HotReloadDefs
 
 IMPLEMENT_MODULE(FHotReloadModule, HotReload);
 
+namespace
+{
+	/**
+	 * Gets editor runs directory.
+	 */
+	FString GetEditorRunsDir()
+	{
+		FString TempDir = FPaths::EngineIntermediateDir();
+
+		return FPaths::Combine(*TempDir, TEXT("EditorRuns"));
+	}
+
+	/**
+	 * Creates a file that informs UBT that the editor is currently running.
+	 */
+	void CreateFileThatIndicatesEditorRunIfNeeded()
+	{
+#if WITH_EDITOR
+		IPlatformFile& FS = IPlatformFile::GetPlatformPhysical();
+
+		FString EditorRunsDir = GetEditorRunsDir();
+		FString FileName = FPaths::Combine(*EditorRunsDir, *FString::Printf(TEXT("%d"), FPlatformProcess::GetCurrentProcessId()));
+
+		if (FS.FileExists(*FileName))
+		{
+			if (!GIsEditor)
+			{
+				FS.DeleteFile(*FileName);
+			}
+		}
+		else
+		{
+			if (GIsEditor)
+			{
+				if (!FS.CreateDirectory(*EditorRunsDir))
+				{
+					return;
+				}
+
+				delete FS.OpenWrite(*FileName); // Touch file.
+			}
+		}
+#endif // WITH_EDITOR
+	}
+
+	/**
+	 * Deletes file left by CreateFileThatIndicatesEditorRunIfNeeded function.
+	 */
+	void DeleteFileThatIndicatesEditorRunIfNeeded()
+	{
+#if WITH_EDITOR
+		IPlatformFile& FS = IPlatformFile::GetPlatformPhysical();
+
+		FString EditorRunsDir = GetEditorRunsDir();
+		FString FileName = FPaths::Combine(*EditorRunsDir, *FString::Printf(TEXT("%d"), FPlatformProcess::GetCurrentProcessId()));
+
+		if (FS.FileExists(*FileName))
+		{
+			FS.DeleteFile(*FileName);
+		}
+#endif // WITH_EDITOR
+	}
+}
+
 void FHotReloadModule::StartupModule()
 {
+	CreateFileThatIndicatesEditorRunIfNeeded();
+
 	bIsHotReloadingFromEditor = false;
 
 #if WITH_ENGINE
@@ -414,6 +480,8 @@ void FHotReloadModule::ShutdownModule()
 {
 	FTicker::GetCoreTicker().RemoveTicker(TickerDelegateHandle);
 	ShutdownHotReloadWatcher();
+
+	DeleteFileThatIndicatesEditorRunIfNeeded();
 }
 
 bool FHotReloadModule::Exec( UWorld* Inworld, const TCHAR* Cmd, FOutputDevice& Ar )
diff --git a/Engine/Source/Developer/MaterialUtilities/Private/MaterialUtilities.cpp b/Engine/Source/Developer/MaterialUtilities/Private/MaterialUtilities.cpp
--- a/Engine/Source/Developer/MaterialUtilities/Private/MaterialUtilities.cpp
+++ b/Engine/Source/Developer/MaterialUtilities/Private/MaterialUtilities.cpp
@@ -255,6 +255,10 @@ public:
 		case MP_Metallic: ResourceId.Usage = EMaterialShaderMapUsage::MaterialExportMetallic; break;
 		case MP_Roughness: ResourceId.Usage = EMaterialShaderMapUsage::MaterialExportRoughness; break;
 		case MP_AmbientOcclusion: ResourceId.Usage = EMaterialShaderMapUsage::MaterialExportAO; break;
+		case MP_EmissiveColor: ResourceId.Usage = EMaterialShaderMapUsage::MaterialExportEmissive; break;
+		default:
+			ensureMsgf(false, TEXT("ExportMaterial has no usage for property %i.  Will likely reuse the normal rendering shader and crash later with a parameter mismatch"), (int32)InPropertyToCompile);
+			break;
 		};
 		
 		CacheShaders(ResourceId, GMaxRHIShaderPlatform, true);
diff --git a/Engine/Source/Editor/LandscapeEditor/Private/LandscapeEdMode.cpp b/Engine/Source/Editor/LandscapeEditor/Private/LandscapeEdMode.cpp
--- a/Engine/Source/Editor/LandscapeEditor/Private/LandscapeEdMode.cpp
+++ b/Engine/Source/Editor/LandscapeEditor/Private/LandscapeEdMode.cpp
@@ -2487,6 +2487,8 @@ ALandscape* FEdModeLandscape::ChangeComponentSetting(int32 NumComponentsX, int32
 	check(NumSubsections > 0);
 	check(SubsectionSizeQuads > 0);
 
+	const int32 NewComponentSizeQuads = NumSubsections * SubsectionSizeQuads;
+
 	ALandscape* Landscape = NULL;
 
 	ULandscapeInfo* LandscapeInfo = CurrentToolTarget.LandscapeInfo.Get();
@@ -2499,8 +2501,8 @@ ALandscape* FEdModeLandscape::ChangeComponentSetting(int32 NumComponentsX, int32
 
 			const int32 OldVertsX = OldMaxX - OldMinX + 1;
 			const int32 OldVertsY = OldMaxY - OldMinY + 1;
-			const int32 NewVertsX = NumComponentsX * NumSubsections * SubsectionSizeQuads + 1;
-			const int32 NewVertsY = NumComponentsY * NumSubsections * SubsectionSizeQuads + 1;
+			const int32 NewVertsX = NumComponentsX * NewComponentSizeQuads + 1;
+			const int32 NewVertsY = NumComponentsY * NewComponentSizeQuads + 1;
 
 			FLandscapeEditDataInterface LandscapeEdit(LandscapeInfo);
 			TArray<uint16> HeightData;
@@ -2508,8 +2510,14 @@ ALandscape* FEdModeLandscape::ChangeComponentSetting(int32 NumComponentsX, int32
 			FVector LandscapeOffset = FVector::ZeroVector;
 			float LandscapeScaleFactor = 1.0f;
 
+			int32 NewMinX, NewMinY, NewMaxX, NewMaxY;
 			if (bResample)
 			{
+				NewMinX = OldMinX / LandscapeInfo->ComponentSizeQuads * NewComponentSizeQuads;
+				NewMinY = OldMinY / LandscapeInfo->ComponentSizeQuads * NewComponentSizeQuads;
+				NewMaxX = NewMinX + NewVertsX - 1;
+				NewMaxY = NewMinY + NewVertsY - 1;
+
 				HeightData.AddZeroed(OldVertsX * OldVertsY * sizeof(uint16));
 
 				// GetHeightData alters its args, so make temp copies to avoid screwing things up
@@ -2535,14 +2543,14 @@ ALandscape* FEdModeLandscape::ChangeComponentSetting(int32 NumComponentsX, int32
 					}
 				}
 
-				LandscapeScaleFactor = (float)OldLandscapeProxy->ComponentSizeQuads / (NumSubsections * SubsectionSizeQuads);
+				LandscapeScaleFactor = (float)OldLandscapeProxy->ComponentSizeQuads / NewComponentSizeQuads;
 			}
 			else
 			{
-				const int32 NewMinX = OldMinX + (OldVertsX - NewVertsX) / 2;
-				const int32 NewMinY = OldMinY + (OldVertsY - NewVertsY) / 2;
-				const int32 NewMaxX = NewMinX + NewVertsX - 1;
-				const int32 NewMaxY = NewMinY + NewVertsY - 1;
+				NewMinX = OldMinX + (OldVertsX - NewVertsX) / 2;
+				NewMinY = OldMinY + (OldVertsY - NewVertsY) / 2;
+				NewMaxX = NewMinX + NewVertsX - 1;
+				NewMaxY = NewMinY + NewVertsY - 1;
 				const int32 RequestedMinX = FMath::Max(OldMinX, NewMinX);
 				const int32 RequestedMinY = FMath::Max(OldMinY, NewMinY);
 				const int32 RequestedMaxX = FMath::Min(OldMaxX, NewMaxX);
@@ -2577,7 +2585,12 @@ ALandscape* FEdModeLandscape::ChangeComponentSetting(int32 NumComponentsX, int32
 					}
 				}
 
+				// offset landscape to component boundary
 				LandscapeOffset = FVector(NewMinX - OldMinX, NewMinY - OldMinY, 0) * OldLandscapeProxy->GetActorScale();
+				NewMinX = 0;
+				NewMinY = 0;
+				NewMaxX = NewVertsX - 1;
+				NewMaxY = NewVertsY - 1;
 			}
 
 			const FVector Location = OldLandscapeProxy->GetActorLocation() + LandscapeOffset;
@@ -2588,7 +2601,7 @@ ALandscape* FEdModeLandscape::ChangeComponentSetting(int32 NumComponentsX, int32
 
 			Landscape->LandscapeMaterial = OldLandscapeProxy->LandscapeMaterial;
 			Landscape->CollisionMipLevel = OldLandscapeProxy->CollisionMipLevel;
-			Landscape->Import(FGuid::NewGuid(), NewVertsX, NewVertsY, NumSubsections*SubsectionSizeQuads, NumSubsections, SubsectionSizeQuads, HeightData.GetData(), *OldLandscapeProxy->ReimportHeightmapFilePath, ImportLayerInfos);
+			Landscape->Import(FGuid::NewGuid(), NewMinX, NewMinY, NewMaxX, NewMaxY, NumSubsections, SubsectionSizeQuads, HeightData.GetData(), *OldLandscapeProxy->ReimportHeightmapFilePath, ImportLayerInfos, ELandscapeImportAlphamapType::Additive);
 
 			Landscape->MaxLODLevel = OldLandscapeProxy->MaxLODLevel;
 			Landscape->ExportLOD = OldLandscapeProxy->ExportLOD;
diff --git a/Engine/Source/Editor/LandscapeEditor/Private/LandscapeEditorDetailCustomization_NewLandscape.cpp b/Engine/Source/Editor/LandscapeEditor/Private/LandscapeEditorDetailCustomization_NewLandscape.cpp
--- a/Engine/Source/Editor/LandscapeEditor/Private/LandscapeEditorDetailCustomization_NewLandscape.cpp
+++ b/Engine/Source/Editor/LandscapeEditor/Private/LandscapeEditorDetailCustomization_NewLandscape.cpp
@@ -863,8 +863,7 @@ FReply FLandscapeEditorDetailCustomization_NewLandscape::OnCreateButtonClicked()
 			LandscapeEdMode->UISettings->ClearImportLandscapeData();
 		}
 
-		const float ComponentSize = QuadsPerComponent;
-		const FVector Offset = FTransform(LandscapeEdMode->UISettings->NewLandscape_Rotation, FVector::ZeroVector, LandscapeEdMode->UISettings->NewLandscape_Scale).TransformVector(FVector(-ComponentCountX * ComponentSize / 2, -ComponentCountY * ComponentSize / 2, 0));
+		const FVector Offset = FTransform(LandscapeEdMode->UISettings->NewLandscape_Rotation, FVector::ZeroVector, LandscapeEdMode->UISettings->NewLandscape_Scale).TransformVector(FVector(-ComponentCountX * QuadsPerComponent / 2, -ComponentCountY * QuadsPerComponent / 2, 0));
 		ALandscape* Landscape = LandscapeEdMode->GetWorld()->SpawnActor<ALandscape>(LandscapeEdMode->UISettings->NewLandscape_Location + Offset, LandscapeEdMode->UISettings->NewLandscape_Rotation);
 
 		//if (LandscapeEdMode->NewLandscapePreviewMode == ENewLandscapePreviewMode::ImportLandscape)
@@ -873,7 +872,7 @@ FReply FLandscapeEditorDetailCustomization_NewLandscape::OnCreateButtonClicked()
 		}
 
 		Landscape->SetActorRelativeScale3D(LandscapeEdMode->UISettings->NewLandscape_Scale);
-		Landscape->Import(FGuid::NewGuid(), SizeX, SizeY, QuadsPerComponent, LandscapeEdMode->UISettings->NewLandscape_SectionsPerComponent, LandscapeEdMode->UISettings->NewLandscape_QuadsPerSection, Data.GetData(), NULL, LayerInfos);
+		Landscape->Import(FGuid::NewGuid(), 0, 0, SizeX-1, SizeY-1, LandscapeEdMode->UISettings->NewLandscape_SectionsPerComponent, LandscapeEdMode->UISettings->NewLandscape_QuadsPerSection, Data.GetData(), NULL, LayerInfos, LandscapeEdMode->UISettings->ImportLandscape_AlphamapType);
 
 		// automatically calculate a lighting LOD that won't crash lightmass (hopefully)
 		// < 2048x2048 -> LOD0
diff --git a/Engine/Source/Editor/LandscapeEditor/Private/LandscapeEditorObject.cpp b/Engine/Source/Editor/LandscapeEditor/Private/LandscapeEditorObject.cpp
--- a/Engine/Source/Editor/LandscapeEditor/Private/LandscapeEditorObject.cpp
+++ b/Engine/Source/Editor/LandscapeEditor/Private/LandscapeEditorObject.cpp
@@ -74,6 +74,7 @@ ULandscapeEditorObject::ULandscapeEditorObject(const FObjectInitializer& ObjectI
 	, NewLandscape_Scale(100, 100, 100)
 	, ImportLandscape_Width(0)
 	, ImportLandscape_Height(0)
+	, ImportLandscape_AlphamapType(ELandscapeImportAlphamapType::Additive)
 
 	// Brush Settings:
 	, BrushRadius(2048.0f)
@@ -282,6 +283,10 @@ void ULandscapeEditorObject::Load()
 	GConfig->GetString(TEXT("LandscapeEdit"), TEXT("NewLandscapeMaterialName"), NewLandscapeMaterialName, GEditorPerProjectIni);
 	NewLandscape_Material = LoadObject<UMaterialInterface>(NULL, *NewLandscapeMaterialName, NULL, LOAD_NoWarn);
 
+	int32 AlphamapType = (uint8)ImportLandscape_AlphamapType;
+	GConfig->GetInt(TEXT("LandscapeEdit"), TEXT("ImportLandscape_AlphamapType"), AlphamapType, GEditorPerProjectIni);
+	ImportLandscape_AlphamapType = (ELandscapeImportAlphamapType)AlphamapType;
+
 	RefreshImportLayersList();
 }
 
@@ -351,6 +356,8 @@ void ULandscapeEditorObject::Save()
 
 	const FString NewLandscapeMaterialName = (NewLandscape_Material != NULL) ? NewLandscape_Material->GetPathName() : FString();
 	GConfig->SetString(TEXT("LandscapeEdit"), TEXT("NewLandscapeMaterialName"), *NewLandscapeMaterialName, GEditorPerProjectIni);
+
+	GConfig->SetInt(TEXT("LandscapeEdit"), TEXT("ImportLandscape_AlphamapType"), (uint8)ImportLandscape_AlphamapType, GEditorPerProjectIni);
 }
 
 // Region
diff --git a/Engine/Source/Editor/LandscapeEditor/Public/LandscapeEditorObject.h b/Engine/Source/Editor/LandscapeEditor/Public/LandscapeEditorObject.h
--- a/Engine/Source/Editor/LandscapeEditor/Public/LandscapeEditorObject.h
+++ b/Engine/Source/Editor/LandscapeEditor/Public/LandscapeEditorObject.h
@@ -472,6 +472,10 @@ class ULandscapeEditorObject : public UObject
 	UPROPERTY(NonTransactional)
 	TArray<uint16> ImportLandscape_Data;
 
+	// Whether the imported alpha maps are to be interpreted as "layered" or "additive" (UE4 uses additive internally)
+	UPROPERTY(Category="New Landscape", EditAnywhere, NonTransactional, EditFixedSize, meta=(DisplayName="Layer Alphamap Type", ShowForTools="NewLandscape"))
+	ELandscapeImportAlphamapType ImportLandscape_AlphamapType;
+
 	// The landscape layers that will be created. Only layer names referenced in the material assigned above are shown here. Modify the material to add more layers.
 	UPROPERTY(Category="New Landscape", EditAnywhere, NonTransactional, EditFixedSize, meta=(DisplayName="Layers", ShowForTools="NewLandscape"))
 	TArray<FLandscapeImportLayer> ImportLandscape_Layers;
diff --git a/Engine/Source/Editor/MeshPaint/Private/MeshPaintEdMode.cpp b/Engine/Source/Editor/MeshPaint/Private/MeshPaintEdMode.cpp
--- a/Engine/Source/Editor/MeshPaint/Private/MeshPaintEdMode.cpp
+++ b/Engine/Source/Editor/MeshPaint/Private/MeshPaintEdMode.cpp
@@ -13,6 +13,7 @@
 #include "Editor/UnrealEd/Public/ObjectTools.h"
 #include "AssetToolsModule.h"
 #include "EditorSupportDelegates.h"
+#include "EditorReimportHandler.h"
 
 //Slate dependencies
 #include "Editor/LevelEditor/Public/LevelEditor.h"
@@ -157,6 +158,10 @@ void FEdModeMeshPaint::Enter()
 
 	}
 
+	// Catch assets if they are about to be (re)imported
+	FEditorDelegates::OnAssetPreImport.AddSP(this, &FEdModeMeshPaint::OnPreImportAsset);
+	FReimportManager::Instance()->OnPreReimport().AddSP(this, &FEdModeMeshPaint::OnPreReimportAsset);
+
 	// Initialize adapter globals
 	FMeshPaintAdapterFactory::InitializeAdapterGlobals();
 
@@ -243,11 +248,7 @@ void FEdModeMeshPaint::Exit()
 	}
 
 	// Remove all adapters
-	for (const auto& Pair : ComponentToAdapterMap)
-	{
-		Pair.Value->OnRemoved();
-	}
-	ComponentToAdapterMap.Empty();
+	RemoveAllGeometryAdapters();
 
 	PaintTargetData.Empty();
 
@@ -262,6 +263,10 @@ void FEdModeMeshPaint::Exit()
 		EndTransaction();
 	}
 
+	// Unbind delegates
+	FReimportManager::Instance()->OnPreReimport().RemoveAll(this);
+	FEditorDelegates::OnAssetPreImport.RemoveAll(this);
+
 	// Call parent implementation
 	FEdMode::Exit();
 }
@@ -962,6 +967,35 @@ void FEdModeMeshPaint::CleanStaleGeometryAdapters(const TArray<UMeshComponent*>&
 	}
 }
 
+void FEdModeMeshPaint::RemoveAllGeometryAdapters()
+{
+	// Remove all components from the map
+	for (auto It = ComponentToAdapterMap.CreateIterator(); It; ++It)
+	{
+		It.Value()->OnRemoved();
+		It.RemoveCurrent();
+	}
+}
+
+void FEdModeMeshPaint::OnPreImportAsset(UFactory* Factory, UClass* Class, UObject* Object, const FName& Name, const TCHAR* Type)
+{
+	if (Class->IsChildOf(UStaticMesh::StaticClass()))
+	{
+		// Remove all geometry adapters to force them all to be recached next time they're required
+		RemoveAllGeometryAdapters();
+	}
+}
+
+void FEdModeMeshPaint::OnPreReimportAsset(UObject* Object)
+{
+	if (Object->IsA(UStaticMesh::StaticClass()))
+	{
+		// Remove all geometry adapters to force them all to be recached next time they're required
+		RemoveAllGeometryAdapters();
+	}
+}
+
+
 /** Paint the mesh that impacts the specified ray */
 void FEdModeMeshPaint::DoPaint( const FVector& InCameraOrigin,
 								const FVector& InRayOrigin,
@@ -2831,21 +2865,23 @@ bool FEdModeMeshPaint::Select( AActor* InActor, bool bInSelected )
 		}
 		else
 		{
-			check(!ComponentToAdapterMap.Contains(MeshComponent));
-			TSharedPtr<IMeshPaintGeometryAdapter> GeomInfo = FMeshPaintAdapterFactory::CreateAdapterForMesh(MeshComponent, PaintingMeshLODIndex, /*TODO: Shouldn't be part of the construction contract: FMeshPaintSettings::Get().UVChannel*/ 0);
-			ComponentToAdapterMap.Add(MeshComponent, GeomInfo);
-			GeomInfo->OnAdded();
-
-			if (FMeshPaintSettings::Get().ResourceType == EMeshPaintResource::Texture)
-			{
-				SetAllTextureOverrides(*GeomInfo.Get(), MeshComponent);
-			}
-			else if (FMeshPaintSettings::Get().ResourceType == EMeshPaintResource::VertexColors)
+			if (!ComponentToAdapterMap.Contains(MeshComponent))
 			{
-				//Painting is done on LOD0 so force the mesh to render only LOD0.
-				ApplyOrRemoveForceBestLOD(*GeomInfo.Get(), MeshComponent, /*bApply=*/ true);
+				TSharedPtr<IMeshPaintGeometryAdapter> GeomInfo = FMeshPaintAdapterFactory::CreateAdapterForMesh(MeshComponent, PaintingMeshLODIndex, /*TODO: Shouldn't be part of the construction contract: FMeshPaintSettings::Get().UVChannel*/ 0);
+				ComponentToAdapterMap.Add(MeshComponent, GeomInfo);
+				GeomInfo->OnAdded();
+
+				if (FMeshPaintSettings::Get().ResourceType == EMeshPaintResource::Texture)
+				{
+					SetAllTextureOverrides(*GeomInfo.Get(), MeshComponent);
+				}
+				else if (FMeshPaintSettings::Get().ResourceType == EMeshPaintResource::VertexColors)
 				{
-					FComponentReregisterContext ReregisterContext(MeshComponent);
+					//Painting is done on LOD0 so force the mesh to render only LOD0.
+					ApplyOrRemoveForceBestLOD(*GeomInfo.Get(), MeshComponent, /*bApply=*/ true);
+					{
+						FComponentReregisterContext ReregisterContext(MeshComponent);
+					}
 				}
 			}
 		}
diff --git a/Engine/Source/Editor/MeshPaint/Private/MeshPaintSplineMeshAdapter.cpp b/Engine/Source/Editor/MeshPaint/Private/MeshPaintSplineMeshAdapter.cpp
--- a/Engine/Source/Editor/MeshPaint/Private/MeshPaintSplineMeshAdapter.cpp
+++ b/Engine/Source/Editor/MeshPaint/Private/MeshPaintSplineMeshAdapter.cpp
@@ -10,9 +10,9 @@
 //////////////////////////////////////////////////////////////////////////
 // FMeshPaintGeometryAdapterForSplineMeshes
 
-bool FMeshPaintGeometryAdapterForSplineMeshes::Construct(UMeshComponent* InComponent, int32 InPaintingMeshLODIndex, int32 InUVChannelIndex)
+bool FMeshPaintGeometryAdapterForSplineMeshes::InitializeMeshData()
 {
-	bool bResult = FMeshPaintGeometryAdapterForStaticMeshes::Construct(InComponent, InPaintingMeshLODIndex, InUVChannelIndex);
+	bool bResult = FMeshPaintGeometryAdapterForStaticMeshes::InitializeMeshData();
 
 	if (bResult)
 	{
diff --git a/Engine/Source/Editor/MeshPaint/Private/MeshPaintSplineMeshAdapter.h b/Engine/Source/Editor/MeshPaint/Private/MeshPaintSplineMeshAdapter.h
--- a/Engine/Source/Editor/MeshPaint/Private/MeshPaintSplineMeshAdapter.h
+++ b/Engine/Source/Editor/MeshPaint/Private/MeshPaintSplineMeshAdapter.h
@@ -11,7 +11,7 @@
 class FMeshPaintGeometryAdapterForSplineMeshes : public FMeshPaintGeometryAdapterForStaticMeshes
 {
 public:
-	virtual bool Construct(UMeshComponent* InComponent, int32 InPaintingMeshLODIndex, int32 InUVChannelIndex) override;
+	virtual bool InitializeMeshData() override;
 	virtual FVector GetMeshVertex(int32 Index) const override;
 
 protected:
diff --git a/Engine/Source/Editor/MeshPaint/Private/MeshPaintStaticMeshAdapter.cpp b/Engine/Source/Editor/MeshPaint/Private/MeshPaintStaticMeshAdapter.cpp
--- a/Engine/Source/Editor/MeshPaint/Private/MeshPaintStaticMeshAdapter.cpp
+++ b/Engine/Source/Editor/MeshPaint/Private/MeshPaintStaticMeshAdapter.cpp
@@ -18,20 +18,43 @@ bool FMeshPaintGeometryAdapterForStaticMeshes::Construct(UMeshComponent* InCompo
 	{
 		if (StaticMeshComponent->StaticMesh != nullptr)
 		{
-			if (InPaintingMeshLODIndex < StaticMeshComponent->StaticMesh->GetNumLODs())
-			{
-				LODModel = &(StaticMeshComponent->StaticMesh->RenderData->LODResources[InPaintingMeshLODIndex]);
-				UVChannelIndex = InUVChannelIndex;
-				Indices = LODModel->IndexBuffer.GetArrayView();
-
-				return (UVChannelIndex >= 0) && (UVChannelIndex < (int32)LODModel->VertexBuffer.GetNumTexCoords());
-			}
+			PaintingMeshLODIndex = InPaintingMeshLODIndex;
+			UVChannelIndex = InUVChannelIndex;
+			StaticMeshComponent->StaticMesh->OnPostMeshBuild().AddRaw(this, &FMeshPaintGeometryAdapterForStaticMeshes::OnPostMeshBuild);
+			const bool bSuccess = InitializeMeshData();
+			return bSuccess;
 		}
 	}
 
 	return false;
 }
 
+FMeshPaintGeometryAdapterForStaticMeshes::~FMeshPaintGeometryAdapterForStaticMeshes()
+{
+	if (StaticMeshComponent != nullptr && StaticMeshComponent->StaticMesh != nullptr)
+	{
+		StaticMeshComponent->StaticMesh->OnPostMeshBuild().RemoveAll(this);
+	}
+}
+
+void FMeshPaintGeometryAdapterForStaticMeshes::OnPostMeshBuild(UStaticMesh* StaticMesh)
+{
+	InitializeMeshData();
+}
+
+bool FMeshPaintGeometryAdapterForStaticMeshes::InitializeMeshData()
+{
+	if (PaintingMeshLODIndex < StaticMeshComponent->StaticMesh->GetNumLODs())
+	{
+		LODModel = &(StaticMeshComponent->StaticMesh->RenderData->LODResources[PaintingMeshLODIndex]);
+		Indices = LODModel->IndexBuffer.GetArrayView();
+
+		return (UVChannelIndex >= 0) && (UVChannelIndex < (int32)LODModel->VertexBuffer.GetNumTexCoords());
+	}
+
+	return false;
+}
+
 void FMeshPaintGeometryAdapterForStaticMeshes::InitializeAdapterGlobals()
 {
 	MeshToComponentMap.Empty();
diff --git a/Engine/Source/Editor/MeshPaint/Private/MeshPaintStaticMeshAdapter.h b/Engine/Source/Editor/MeshPaint/Private/MeshPaintStaticMeshAdapter.h
--- a/Engine/Source/Editor/MeshPaint/Private/MeshPaintStaticMeshAdapter.h
+++ b/Engine/Source/Editor/MeshPaint/Private/MeshPaintStaticMeshAdapter.h
@@ -13,6 +13,8 @@ public:
 	static void InitializeAdapterGlobals();
 
 	virtual bool Construct(UMeshComponent* InComponent, int32 InPaintingMeshLODIndex, int32 InUVChannelIndex) override;
+	virtual ~FMeshPaintGeometryAdapterForStaticMeshes();
+	virtual bool InitializeMeshData() override;
 	virtual void OnAdded() override;
 	virtual void OnRemoved() override;
 	virtual int32 GetNumTexCoords() const override;
@@ -28,6 +30,9 @@ public:
 	virtual FVector GetMeshVertex(int32 Index) const override;
 
 protected:
+
+	void OnPostMeshBuild(UStaticMesh* StaticMesh);
+
 	struct FStaticMeshReferencers
 	{
 		FStaticMeshReferencers()
@@ -56,6 +61,7 @@ protected:
 
 	UStaticMeshComponent* StaticMeshComponent;
 	FStaticMeshLODResources* LODModel;
+	int32 PaintingMeshLODIndex;
 	int32 UVChannelIndex;
 	FIndexArrayView Indices;
 };
diff --git a/Engine/Source/Editor/MeshPaint/Public/MeshPaintEdMode.h b/Engine/Source/Editor/MeshPaint/Public/MeshPaintEdMode.h
--- a/Engine/Source/Editor/MeshPaint/Public/MeshPaintEdMode.h
+++ b/Engine/Source/Editor/MeshPaint/Public/MeshPaintEdMode.h
@@ -750,6 +750,15 @@ private:
 	/** Removes stale geometry adapters from the cache (those associated with unselected components) */
 	void CleanStaleGeometryAdapters(const TArray<UMeshComponent*>& ValidComponents);
 
+	/** Removes all geometry adapters from the cache */
+	void RemoveAllGeometryAdapters();
+
+	/** Called when an asset is about to be imported */
+	void OnPreImportAsset(UFactory* Factory, UClass* Class, UObject* Object, const FName& Name, const TCHAR* Type);
+
+	/** Called when an asset is about to be reimported */
+	void OnPreReimportAsset(UObject* Object);
+
 private:
 
 	/** Whether we're currently painting */
diff --git a/Engine/Source/Editor/MeshPaint/Public/MeshPaintModule.h b/Engine/Source/Editor/MeshPaint/Public/MeshPaintModule.h
--- a/Engine/Source/Editor/MeshPaint/Public/MeshPaintModule.h
+++ b/Engine/Source/Editor/MeshPaint/Public/MeshPaintModule.h
@@ -11,6 +11,7 @@ class IMeshPaintGeometryAdapter
 {
 public:
 	virtual bool Construct(UMeshComponent* InComponent, int32 InPaintingMeshLODIndex, int32 InUVChannelIndex) = 0;
+	virtual bool InitializeMeshData() = 0;
 	virtual void OnAdded() = 0;
 	virtual void OnRemoved() = 0;
 	virtual int32 GetNumTexCoords() const = 0;
diff --git a/Engine/Source/Editor/MovieSceneTools/Private/CommonMovieSceneTools.cpp b/Engine/Source/Editor/MovieSceneTools/Private/CommonMovieSceneTools.cpp
new file mode 100644
--- /dev/null
+++ b/Engine/Source/Editor/MovieSceneTools/Private/CommonMovieSceneTools.cpp
@@ -0,0 +1,53 @@
+// Copyright 1998-2015 Epic Games, Inc. All Rights Reserved.
+
+#include "MovieSceneToolsPrivatePCH.h"
+#include "CommonMovieSceneTools.h"
+
+/** Converts a matinee interpolation mode to its equivalent rich curve interpolation mode. */
+ERichCurveInterpMode MatineeImportTools::MatineeInterpolationToRichCurveInterpolation( EInterpCurveMode CurveMode )
+{
+	switch ( CurveMode )
+	{
+	case CIM_Constant:
+		return ERichCurveInterpMode::RCIM_Constant;
+	case CIM_CurveAuto:
+	case CIM_CurveAutoClamped:
+	case CIM_CurveBreak:
+	case CIM_CurveUser:
+		return ERichCurveInterpMode::RCIM_Cubic;
+	case CIM_Linear:
+		return ERichCurveInterpMode::RCIM_Linear;
+	default:
+		return ERichCurveInterpMode::RCIM_None;
+	}
+}
+
+
+/** Converts a matinee interpolation mode to its equivalent rich curve tangent mode. */
+ERichCurveTangentMode MatineeImportTools::MatineeInterpolationToRichCurveTangent( EInterpCurveMode CurveMode )
+{
+	switch ( CurveMode )
+	{
+	case CIM_CurveAuto:
+	case CIM_CurveAutoClamped:
+		return ERichCurveTangentMode::RCTM_Auto;
+	case CIM_CurveBreak:
+		return ERichCurveTangentMode::RCTM_Break;
+	case CIM_CurveUser:
+		return ERichCurveTangentMode::RCTM_User;
+	default:
+		return ERichCurveTangentMode::RCTM_None;
+	}
+}
+
+
+/** Adds a key to a rich curve based on matinee curve key data. */
+void MatineeImportTools::SetOrAddKey( FRichCurve& Curve, float Time, float Value, float ArriveTangent, float LeaveTangent, EInterpCurveMode MatineeInterpMode )
+{
+	FKeyHandle KeyHandle = Curve.AddKey( Time, Value, false, Curve.FindKey( Time ) );
+	FRichCurveKey& Key = Curve.GetKey( KeyHandle );
+	Key.ArriveTangent = ArriveTangent;
+	Key.LeaveTangent = LeaveTangent;
+	Key.InterpMode = MatineeInterpolationToRichCurveInterpolation( MatineeInterpMode );
+	Key.TangentMode = MatineeInterpolationToRichCurveTangent( MatineeInterpMode );
+}
diff --git a/Engine/Source/Editor/MovieSceneTools/Private/TrackEditors/TransformTrackEditor.cpp b/Engine/Source/Editor/MovieSceneTools/Private/TrackEditors/TransformTrackEditor.cpp
--- a/Engine/Source/Editor/MovieSceneTools/Private/TrackEditors/TransformTrackEditor.cpp
+++ b/Engine/Source/Editor/MovieSceneTools/Private/TrackEditors/TransformTrackEditor.cpp
@@ -14,6 +14,7 @@
 #include "MovieSceneToolHelpers.h"
 #include "MovieSceneTrackEditor.h"
 #include "TransformTrackEditor.h"
+#include "Matinee/InterpTrackMove.h"
 
 
 #define LOCTEXT_NAMESPACE "MovieScene_TransformTrack"
@@ -244,6 +245,64 @@ bool F3DTransformTrackEditor::SupportsType( TSubclassOf<UMovieSceneTrack> Type )
 }
 
 
+void PasteInterpMoveTrack( UInterpTrackMove* MoveTrack, UMovieScene3DTransformTrack* TransformTrack )
+{
+	float KeyTime = MoveTrack->GetKeyframeTime( 0 );
+	UMovieScene3DTransformSection* Section = Cast<UMovieScene3DTransformSection>(TransformTrack->FindOrAddSection( KeyTime ));
+	float SectionMin = Section->GetStartTime();
+	float SectionMax = Section->GetEndTime();
+
+	FRichCurve& TranslationXCurve = Section->GetTranslationCurve( EAxis::X );
+	FRichCurve& TranslationYCurve = Section->GetTranslationCurve( EAxis::Y );
+	FRichCurve& TranslationZCurve = Section->GetTranslationCurve( EAxis::Z );
+	for ( const auto& Point : MoveTrack->PosTrack.Points )
+	{
+		MatineeImportTools::SetOrAddKey( TranslationXCurve, Point.InVal, Point.OutVal.X, Point.ArriveTangent.X, Point.LeaveTangent.X, Point.InterpMode );
+		MatineeImportTools::SetOrAddKey( TranslationYCurve, Point.InVal, Point.OutVal.Y, Point.ArriveTangent.Y, Point.LeaveTangent.Y, Point.InterpMode );
+		MatineeImportTools::SetOrAddKey( TranslationZCurve, Point.InVal, Point.OutVal.Z, Point.ArriveTangent.Z, Point.LeaveTangent.Z, Point.InterpMode );
+		SectionMin = FMath::Min( SectionMin, Point.InVal );
+		SectionMax = FMath::Max( SectionMax, Point.InVal );
+	}
+
+	FRichCurve& RotationXCurve = Section->GetRotationCurve( EAxis::X );
+	FRichCurve& RotationYCurve = Section->GetRotationCurve( EAxis::Y );
+	FRichCurve& RotationZCurve = Section->GetRotationCurve( EAxis::Z );
+	for ( const auto& Point : MoveTrack->EulerTrack.Points )
+	{
+		MatineeImportTools::SetOrAddKey( RotationXCurve, Point.InVal, Point.OutVal.X, Point.ArriveTangent.X, Point.LeaveTangent.X, Point.InterpMode );
+		MatineeImportTools::SetOrAddKey( RotationYCurve, Point.InVal, Point.OutVal.Y, Point.ArriveTangent.Y, Point.LeaveTangent.Y, Point.InterpMode );
+		MatineeImportTools::SetOrAddKey( RotationZCurve, Point.InVal, Point.OutVal.Z, Point.ArriveTangent.Z, Point.LeaveTangent.Z, Point.InterpMode );
+		SectionMin = FMath::Min( SectionMin, Point.InVal );
+		SectionMax = FMath::Max( SectionMax, Point.InVal );
+	}
+
+	Section->SetStartTime( SectionMin );
+	Section->SetEndTime( SectionMax );
+}
+
+
+void F3DTransformTrackEditor::BuildTrackContextMenu( FMenuBuilder& MenuBuilder, UMovieSceneTrack* Track )
+{
+	UInterpTrackMove* MoveTrack = nullptr;
+	for ( UObject* CopyPasteObject : GUnrealEd->MatineeCopyPasteBuffer )
+	{
+		MoveTrack = Cast<UInterpTrackMove>( CopyPasteObject );
+		if ( MoveTrack != nullptr )
+		{
+			break;
+		}
+	}
+	UMovieScene3DTransformTrack* TransformTrack = Cast<UMovieScene3DTransformTrack>( Track );
+	MenuBuilder.AddMenuEntry(
+		NSLOCTEXT( "Sequencer", "PasteMatineeTrack", "Paste Matinee Move Track" ),
+		NSLOCTEXT( "Sequencer", "PasteMatineeTrackTooltip", "Pastes keys from a Matinee move track into this track." ),
+		FSlateIcon(),
+		FUIAction(
+			FExecuteAction::CreateStatic( &PasteInterpMoveTrack, MoveTrack, TransformTrack ),
+			FCanExecuteAction::CreateLambda( [=]()->bool { return MoveTrack != nullptr && MoveTrack->GetNumKeys() > 0 && TransformTrack != nullptr; } ) ) );
+}
+
+
 TSharedRef<ISequencerSection> F3DTransformTrackEditor::MakeSectionInterface( UMovieSceneSection& SectionObject, UMovieSceneTrack& Track )
 {
 	check( SupportsType( SectionObject.GetOuter()->GetClass() ) );
diff --git a/Engine/Source/Editor/MovieSceneTools/Private/TrackEditors/TransformTrackEditor.h b/Engine/Source/Editor/MovieSceneTools/Private/TrackEditors/TransformTrackEditor.h
--- a/Engine/Source/Editor/MovieSceneTools/Private/TrackEditors/TransformTrackEditor.h
+++ b/Engine/Source/Editor/MovieSceneTools/Private/TrackEditors/TransformTrackEditor.h
@@ -44,6 +44,7 @@ public:
 	virtual TSharedRef<ISequencerSection> MakeSectionInterface( UMovieSceneSection& SectionObject, UMovieSceneTrack& Track ) override;
 	virtual void OnRelease() override;
 	virtual bool SupportsType( TSubclassOf<UMovieSceneTrack> Type ) const override;
+	virtual void BuildTrackContextMenu( FMenuBuilder& MenuBuilder, UMovieSceneTrack* Track ) override;
 
 private:
 
diff --git a/Engine/Source/Editor/MovieSceneTools/Public/CommonMovieSceneTools.h b/Engine/Source/Editor/MovieSceneTools/Public/CommonMovieSceneTools.h
--- a/Engine/Source/Editor/MovieSceneTools/Public/CommonMovieSceneTools.h
+++ b/Engine/Source/Editor/MovieSceneTools/Public/CommonMovieSceneTools.h
@@ -53,3 +53,18 @@ private:
 	/** The number of pixels in the view range */
 	float PixelsPerInput;
 };
+
+
+class MatineeImportTools
+{
+public:
+
+	/** Converts a matinee interpolation mode to its equivalent rich curve interpolation mode. */
+	static ERichCurveInterpMode MatineeInterpolationToRichCurveInterpolation( EInterpCurveMode CurveMode );
+
+	/** Converts a matinee interpolation mode to its equivalent rich curve tangent mode. */
+	static ERichCurveTangentMode MatineeInterpolationToRichCurveTangent( EInterpCurveMode CurveMode );
+
+	/** Adds a key to a rich curve based on matinee curve key data. */
+	static void SetOrAddKey( FRichCurve& Curve, float Time, float Value, float ArriveTangent, float LeaveTangent, EInterpCurveMode MatineeInterpMode );
+};
diff --git a/Engine/Source/Editor/Sequencer/Private/DisplayNodes/SequencerTrackNode.cpp b/Engine/Source/Editor/Sequencer/Private/DisplayNodes/SequencerTrackNode.cpp
--- a/Engine/Source/Editor/Sequencer/Private/DisplayNodes/SequencerTrackNode.cpp
+++ b/Engine/Source/Editor/Sequencer/Private/DisplayNodes/SequencerTrackNode.cpp
@@ -169,3 +169,9 @@ ESequencerNode::Type FTrackNode::GetType() const
 {
 	return ESequencerNode::Track;
 }
+
+void FTrackNode::BuildContextMenu( FMenuBuilder& MenuBuilder )
+{
+	AssociatedEditor.BuildTrackContextMenu( MenuBuilder, AssociatedType.Get() );
+	FSequencerDisplayNode::BuildContextMenu( MenuBuilder );
+}
