Disable Oodle by default, enable in projects that use it.

diff --git a/Engine/Plugins/Runtime/PacketHandlers/CompressionComponents/Oodle/Oodle.uplugin b/Engine/Plugins/Runtime/PacketHandlers/CompressionComponents/Oodle/Oodle.uplugin
--- a/Engine/Plugins/Runtime/PacketHandlers/CompressionComponents/Oodle/Oodle.uplugin
+++ b/Engine/Plugins/Runtime/PacketHandlers/CompressionComponents/Oodle/Oodle.uplugin
@@ -9,7 +9,7 @@
 	"EngineVersion" : "4.4.0",
 	"Description" : "Realtime network compression using Oodle technology",
 	"Category" : "Networking",
-	"EnabledByDefault" : true,
+	"EnabledByDefault" : false,
     "IsBetaVersion": false,
 	
 	"Modules" :
