Fixing CIS

diff --git a/Engine/Plugins/Runtime/PacketHandlers/CompressionComponents/Oodle/Source/OodleHandlerComponent/Classes/OodleTrainerCommandlet.h b/Engine/Plugins/Runtime/PacketHandlers/CompressionComponents/Oodle/Source/OodleHandlerComponent/Classes/OodleTrainerCommandlet.h
--- a/Engine/Plugins/Runtime/PacketHandlers/CompressionComponents/Oodle/Source/OodleHandlerComponent/Classes/OodleTrainerCommandlet.h
+++ b/Engine/Plugins/Runtime/PacketHandlers/CompressionComponents/Oodle/Source/OodleHandlerComponent/Classes/OodleTrainerCommandlet.h
@@ -154,16 +154,16 @@ private:
 	TMap<FArchive*, FString> MergeMap;
 
 
-	TArray<void*>	DictionaryPackets;
+	TArray<uint8*>	DictionaryPackets;
 	TArray<S32>		DictionaryPacketSizes;
 	uint32			DictionaryPacketBytes;
-	TArray<void*>	DictionaryTestPackets;
+	TArray<uint8*>	DictionaryTestPackets;
 	TArray<S32>		DictionaryTestPacketSizes;
 	uint32			DictionaryTestPacketBytes;
-	TArray<void*>	TrainerPackets;
+	TArray<uint8*>	TrainerPackets;
 	TArray<S32>		TrainerPacketSizes;
 	uint32			TrainerPacketBytes;
-	TArray<void*>	CompressionTestPackets;
+	TArray<uint8*>	CompressionTestPackets;
 	TArray<S32>		CompressionTestPacketSizes;
 	uint32			CompressionTestPacketBytes;
 
diff --git a/Engine/Plugins/Runtime/PacketHandlers/CompressionComponents/Oodle/Source/OodleHandlerComponent/Private/OodleTrainerCommandlet.cpp b/Engine/Plugins/Runtime/PacketHandlers/CompressionComponents/Oodle/Source/OodleHandlerComponent/Private/OodleTrainerCommandlet.cpp
--- a/Engine/Plugins/Runtime/PacketHandlers/CompressionComponents/Oodle/Source/OodleHandlerComponent/Private/OodleTrainerCommandlet.cpp
+++ b/Engine/Plugins/Runtime/PacketHandlers/CompressionComponents/Oodle/Source/OodleHandlerComponent/Private/OodleTrainerCommandlet.cpp
@@ -694,7 +694,7 @@ bool FOodleDictionaryGenerator::ReadPackets(FString InputCaptureFiles, bool bMer
 
 				check(ListIdx < ListSplitCount);
 
-				TArray<void*>& TargetPacketList =
+				TArray<uint8*>& TargetPacketList =
 					(
 						ListIdx == 0 ?			DictionaryPackets :
 						ListIdx == 1 ?			(bDictionaryTestOverflow ? TrainerPackets : DictionaryTestPackets) :
@@ -719,7 +719,7 @@ bool FOodleDictionaryGenerator::ReadPackets(FString InputCaptureFiles, bool bMer
 					);
 
 
-				void* PacketDest = new uint8[PacketSize];
+				uint8* PacketDest = new uint8[PacketSize];
 
 				TargetPacketList.Add(PacketDest);
 
@@ -756,7 +756,7 @@ bool FOodleDictionaryGenerator::ReadPackets(FString InputCaptureFiles, bool bMer
 bool FOodleDictionaryGenerator::GenerateAndWriteDictionary()
 {
 	bool bSuccess = false;
-	void* NewDictionaryData = new uint8[DictionarySize];
+	uint8* NewDictionaryData = new uint8[DictionarySize];
 
 
 	UE_LOG(OodleHandlerComponentLog, Log, TEXT("Beginning dictionary generation..."));
@@ -889,7 +889,7 @@ bool FOodleDictionaryGenerator::GenerateAndWriteDictionary()
 
 void FOodleDictionaryGenerator::Cleanup()
 {
-	for (void* CurPacket : DictionaryPackets)
+	for (uint8* CurPacket : DictionaryPackets)
 	{
 		if (CurPacket != nullptr)
 		{
@@ -897,7 +897,7 @@ void FOodleDictionaryGenerator::Cleanup()
 		}
 	}
 
-	for (void* CurPacket : DictionaryTestPackets)
+	for (uint8* CurPacket : DictionaryTestPackets)
 	{
 		if (CurPacket != nullptr)
 		{
@@ -905,7 +905,7 @@ void FOodleDictionaryGenerator::Cleanup()
 		}
 	}
 
-	for (void* CurPacket : TrainerPackets)
+	for (uint8* CurPacket : TrainerPackets)
 	{
 		if (CurPacket != nullptr)
 		{
@@ -913,7 +913,7 @@ void FOodleDictionaryGenerator::Cleanup()
 		}
 	}
 
-	for (void* CurPacket : CompressionTestPackets)
+	for (uint8* CurPacket : CompressionTestPackets)
 	{
 		if (CurPacket != nullptr)
 		{
