Fixed shadowed variable

diff --git a/Engine/Plugins/Runtime/PacketHandlers/CompressionComponents/Oodle/Source/OodleHandlerComponent/Private/OodleTrainerCommandlet.cpp b/Engine/Plugins/Runtime/PacketHandlers/CompressionComponents/Oodle/Source/OodleHandlerComponent/Private/OodleTrainerCommandlet.cpp
--- a/Engine/Plugins/Runtime/PacketHandlers/CompressionComponents/Oodle/Source/OodleHandlerComponent/Private/OodleTrainerCommandlet.cpp
+++ b/Engine/Plugins/Runtime/PacketHandlers/CompressionComponents/Oodle/Source/OodleHandlerComponent/Private/OodleTrainerCommandlet.cpp
@@ -636,7 +636,6 @@ bool FOodleDictionaryGenerator::ReadPackets(FString InputCaptureFiles, bool bMer
 		}
 
 		// Now translate the randomized index, back into a sequential index, which maps random packets to each list-type (from type 0-4)
-		bool bCompressionTestPackets = false;	// @todo #JohnB: Enable from commandline
 		int32 ListSplitCount = 3 + (bCompressionTestPackets ? 1 : 0);
 		TArray<uint8> PacketToListMap;
 
