More CIS fixes for Oodle

diff --git a/Engine/Plugins/Runtime/PacketHandlers/CompressionComponents/Oodle/Source/OodleHandlerComponent/Private/OodleArchives.cpp b/Engine/Plugins/Runtime/PacketHandlers/CompressionComponents/Oodle/Source/OodleHandlerComponent/Private/OodleArchives.cpp
--- a/Engine/Plugins/Runtime/PacketHandlers/CompressionComponents/Oodle/Source/OodleHandlerComponent/Private/OodleArchives.cpp
+++ b/Engine/Plugins/Runtime/PacketHandlers/CompressionComponents/Oodle/Source/OodleHandlerComponent/Private/OodleArchives.cpp
@@ -28,7 +28,7 @@ void FOodleArchiveBase::FOodleCompressedData::Serialize(FArchive& Ar)
  * FOodleArchiveBase
  */
 
-bool FOodleArchiveBase::SerializeOodleCompressData(FOodleCompressedData& OutDataInfo, void* Data, uint32 DataBytes)
+bool FOodleArchiveBase::SerializeOodleCompressData(FOodleCompressedData& OutDataInfo, uint8* Data, uint32 DataBytes)
 {
 	bool bSuccess = true;
 
@@ -41,7 +41,7 @@ bool FOodleArchiveBase::SerializeOodleCompressData(FOodleCompressedData& OutData
 		OutDataInfo.DecompressedLength.Set(*this, DataBytes);
 	
 		uint32 CompressBufferLen = OodleLZ_GetCompressedBufferSizeNeeded(DataBytes);
-		void* CompressBuffer = new uint8[CompressBufferLen];
+		uint8* CompressBuffer = new uint8[CompressBufferLen];
 
 		SINTa OodleLen = OodleLZ_Compress(OodleLZ_Compressor_LZNIB, Data, (SINTa)DataBytes, CompressBuffer,
 											OodleLZ_CompressionLevel_Optimal);
@@ -65,7 +65,7 @@ bool FOodleArchiveBase::SerializeOodleCompressData(FOodleCompressedData& OutData
 	return bSuccess;
 }
 
-bool FOodleArchiveBase::SerializeOodleDecompressData(FOodleCompressedData& DataInfo, void*& OutData, uint32& OutDataBytes)
+bool FOodleArchiveBase::SerializeOodleDecompressData(FOodleCompressedData& DataInfo, uint8*& OutData, uint32& OutDataBytes)
 {
 	check(OutData == nullptr);
 
@@ -82,8 +82,8 @@ bool FOodleArchiveBase::SerializeOodleDecompressData(FOodleCompressedData& DataI
 	{
 		SeekPush(DataOffset);
 
-		void* CompressedData = new uint8[CompressedLength];
-		void* DecompressedData = new uint8[DecompressedLength];
+		uint8* CompressedData = new uint8[CompressedLength];
+		uint8* DecompressedData = new uint8[DecompressedLength];
 
 		InnerArchive.Serialize(CompressedData, CompressedLength);
 
@@ -367,8 +367,8 @@ void FOodleDictionaryArchive::SerializeHeader()
 	Header.SerializeHeader(*this);
 }
 
-void FOodleDictionaryArchive::SerializeDictionaryAndState(void*& DictionaryData, uint32& DictionaryBytes,
-															void*& CompactCompresorState, uint32& CompactCompressorStateBytes)
+void FOodleDictionaryArchive::SerializeDictionaryAndState(uint8*& DictionaryData, uint32& DictionaryBytes,
+															uint8*& CompactCompresorState, uint32& CompactCompressorStateBytes)
 {
 	if (IsLoading())
 	{
diff --git a/Engine/Plugins/Runtime/PacketHandlers/CompressionComponents/Oodle/Source/OodleHandlerComponent/Private/OodleHandlerComponent.cpp b/Engine/Plugins/Runtime/PacketHandlers/CompressionComponents/Oodle/Source/OodleHandlerComponent/Private/OodleHandlerComponent.cpp
--- a/Engine/Plugins/Runtime/PacketHandlers/CompressionComponents/Oodle/Source/OodleHandlerComponent/Private/OodleHandlerComponent.cpp
+++ b/Engine/Plugins/Runtime/PacketHandlers/CompressionComponents/Oodle/Source/OodleHandlerComponent/Private/OodleHandlerComponent.cpp
@@ -131,7 +131,7 @@ void FOodleNetStats::UpdateStats(float DeltaTime)
  * FOodleDictionary
  */
 
-FOodleDictionary::FOodleDictionary(uint32 InHashTableSize, void* InDictionaryData, uint32 InDictionarySize,
+FOodleDictionary::FOodleDictionary(uint32 InHashTableSize, uint8* InDictionaryData, uint32 InDictionarySize,
 									OodleNetwork1_Shared* InSharedDictionary, uint32 InSharedDictionarySize,
 									OodleNetwork1UDP_State* InInitialCompressorState, uint32 InCompressorStateSize)
 	: HashTableSize(InHashTableSize)
@@ -161,7 +161,7 @@ FOodleDictionary::~FOodleDictionary()
 
 	if (DictionaryData != nullptr)
 	{
-		delete DictionaryData;
+		delete [] DictionaryData;
 	}
 
 	if (SharedDictionary != nullptr)
@@ -427,9 +427,9 @@ void OodleHandlerComponent::InitializeDictionary(FString FilePath, TSharedPtr<FO
 		{
 			FOodleDictionaryArchive BoundArc(*ReadArc);
 
-			void* DictionaryData = nullptr;
+			uint8* DictionaryData = nullptr;
 			uint32 DictionaryBytes = 0;
-			void* CompactCompressorState = nullptr;
+			uint8* CompactCompressorState = nullptr;
 			uint32 CompactCompressorStateBytes = 0;
 
 			BoundArc.SerializeHeader();
@@ -464,7 +464,7 @@ void OodleHandlerComponent::InitializeDictionary(FString FilePath, TSharedPtr<FO
 
 				if (CompactCompressorState != nullptr)
 				{
-					delete CompactCompressorState;
+					delete [] CompactCompressorState;
 				}
 			}
 			else
@@ -473,12 +473,12 @@ void OodleHandlerComponent::InitializeDictionary(FString FilePath, TSharedPtr<FO
 
 				if (DictionaryData != nullptr)
 				{
-					delete DictionaryData;
+					delete [] DictionaryData;
 				}
 
 				if (CompactCompressorState != nullptr)
 				{
-					delete CompactCompressorState;
+					delete [] CompactCompressorState;
 				}
 			}
 
diff --git a/Engine/Plugins/Runtime/PacketHandlers/CompressionComponents/Oodle/Source/OodleHandlerComponent/Private/OodleTrainerCommandlet.cpp b/Engine/Plugins/Runtime/PacketHandlers/CompressionComponents/Oodle/Source/OodleHandlerComponent/Private/OodleTrainerCommandlet.cpp
--- a/Engine/Plugins/Runtime/PacketHandlers/CompressionComponents/Oodle/Source/OodleHandlerComponent/Private/OodleTrainerCommandlet.cpp
+++ b/Engine/Plugins/Runtime/PacketHandlers/CompressionComponents/Oodle/Source/OodleHandlerComponent/Private/OodleTrainerCommandlet.cpp
@@ -822,7 +822,7 @@ bool FOodleDictionaryGenerator::GenerateAndWriteDictionary()
 		bSuccess = bSuccess && OutputFile.SerializeOodleCompressData(OutputFile.Header.DictionaryData, NewDictionaryData,
 																		DictionarySize);
 
-		bSuccess = bSuccess && OutputFile.SerializeOodleCompressData(OutputFile.Header.CompressorData, CompactCompressorState,
+		bSuccess = bSuccess && OutputFile.SerializeOodleCompressData(OutputFile.Header.CompressorData, (uint8*)CompactCompressorState,
 																		CompactCompressorStateBytes);
 
 
diff --git a/Engine/Plugins/Runtime/PacketHandlers/CompressionComponents/Oodle/Source/OodleHandlerComponent/Public/OodleArchives.h b/Engine/Plugins/Runtime/PacketHandlers/CompressionComponents/Oodle/Source/OodleHandlerComponent/Public/OodleArchives.h
--- a/Engine/Plugins/Runtime/PacketHandlers/CompressionComponents/Oodle/Source/OodleHandlerComponent/Public/OodleArchives.h
+++ b/Engine/Plugins/Runtime/PacketHandlers/CompressionComponents/Oodle/Source/OodleHandlerComponent/Public/OodleArchives.h
@@ -205,7 +205,7 @@ public:
 	 * @param DataBytes		The length of Data
 	 * @return				Whether or not compression was successful
 	 */
-	bool SerializeOodleCompressData(FOodleCompressedData& OutDataInfo, void* Data, uint32 DataBytes);
+	bool SerializeOodleCompressData(FOodleCompressedData& OutDataInfo, uint8* Data, uint32 DataBytes);
 
 	/**
 	 * Decompresses the data referenced by the input FOodleCompressedData struct
@@ -215,7 +215,7 @@ public:
 	 * @param OutDataBytes	Outputs the size of OutData
 	 * @return				Whether or not decompression was successful
 	 */
-	bool SerializeOodleDecompressData(FOodleCompressedData& DataInfo, void*& OutData, uint32& OutDataBytes);
+	bool SerializeOodleDecompressData(FOodleCompressedData& DataInfo, uint8*& OutData, uint32& OutDataBytes);
 
 
 	/**
@@ -434,8 +434,8 @@ public:
 	 * @param CompactCompressorState		Pointer to compressor data, when saving, and reference to nullptr void*, when loading
 	 * @param CompactCompressorStateBytes	Inputs/Outputs the size of the compressor data
 	 */
-	void SerializeDictionaryAndState(void*& DictionaryData, uint32& DictionaryBytes,
-										void*& CompactCompresorState, uint32& CompactCompressorStateBytes);
+	void SerializeDictionaryAndState(uint8*& DictionaryData, uint32& DictionaryBytes,
+										uint8*& CompactCompresorState, uint32& CompactCompressorStateBytes);
 };
 
 #endif
diff --git a/Engine/Plugins/Runtime/PacketHandlers/CompressionComponents/Oodle/Source/OodleHandlerComponent/Public/OodleHandlerComponent.h b/Engine/Plugins/Runtime/PacketHandlers/CompressionComponents/Oodle/Source/OodleHandlerComponent/Public/OodleHandlerComponent.h
--- a/Engine/Plugins/Runtime/PacketHandlers/CompressionComponents/Oodle/Source/OodleHandlerComponent/Public/OodleHandlerComponent.h
+++ b/Engine/Plugins/Runtime/PacketHandlers/CompressionComponents/Oodle/Source/OodleHandlerComponent/Public/OodleHandlerComponent.h
@@ -205,7 +205,7 @@ struct FOodleDictionary
 	uint32 HashTableSize;
 
 	/** The raw dictionary data */
-	void* DictionaryData;
+	uint8* DictionaryData;
 
 	/** The size of the dictionary */
 	uint32 DictionarySize;
@@ -236,7 +236,7 @@ public:
 	/**
 	 * Base constructor
 	 */
-	FOodleDictionary(uint32 InHashTableSize, void* InDictionaryData, uint32 InDictionarySize, OodleNetwork1_Shared* InSharedDictionary,
+	FOodleDictionary(uint32 InHashTableSize, uint8* InDictionaryData, uint32 InDictionarySize, OodleNetwork1_Shared* InSharedDictionary,
 						uint32 InSharedDictionarySize, OodleNetwork1UDP_State* InInitialCompressorState, uint32 InCompressorStateSize);
 
 	/**
