Copying //UE4/Orion-Staging to //UE4/Main (Source: //Orion/Dev-General @ 2845681)

==========================
MAJOR FEATURES + CHANGES
==========================

Change 2845644 on 2016/01/27 by [REDACTED_PERSON]

	Clear marker sync flag after creating tick record, add more information to checks incase issue occurs again

	#tests in editor tests, bot match.

Change 2845613 on 2016/01/27 by [REDACTED_PERSON]

	Latest network profiler binaries

	#tests run profiler

Change 2845595 on 2016/01/27 by [REDACTED_PERSON]

	Fixed pathfollowing's block detection using wrong distance when testing for blockage #UE4

	#test golden path

Change 2845593 on 2016/01/27 by [REDACTED_PERSON]

	Added support for setting and choosing filmbacks and lenses for cinematic cameras.
	- New CineCameraComponent and CineCameraActor classes
	- can define filmback and lens presets via ini file
	- details customizations for filmback and lens selection
	- added prototype set of filmbacks and lenses (primes and zooms)
	- Camera details customization now gracefully handles when CameraSettings category is hidden
	- example sequencer usage is content/developers/jeff.farris/CineCams/CineCamTestMap

	#tests editor

Change 2845585 on 2016/01/27 by [REDACTED_PERSON]

	Don't fool with connected state if we're early outing from the OS intercepting controller events.  This fixes some missing delegates.
	Fixes cert bug about controller disconnect screen staying up permanently
	#test Turning off controller, turning on again.

Change 2845528 on 2016/01/27 by [REDACTED_PERSON]

	Sequencer: Fix new spawnables not immediately getting an object binding. This was resulted in a missing +Track->Animation when first creating a spawnable and duplicate transform keys.

	#tests Add spawnable, +Track->Animation exists

Change 2845483 on 2016/01/27 by [REDACTED_PERSON]

	Sequencer: Fixed MaximizedViewport not getting cleared/restored correctly
	#jria UE-26016
	#tests Tested the viewports

Change 2845421 on 2016/01/27 by [REDACTED_PERSON]

	Sequencer: Implemented go-to feature

	#TESTS Editor

Change 2845407 on 2016/01/27 by [REDACTED_PERSON]

	Sequencer: Moved SetViewRange() into ISequencer and made it public

	#TESTS none

Change 2845404 on 2016/01/27 by [REDACTED_PERSON]

	Sequencer: Fixed cinematic viewport not updating when dragging transport range
	#tests Scrubbed the timeline

Change 2845396 on 2016/01/27 by [REDACTED_PERSON]

	#OSS #Purchase #Store #PS4
	- Minor log cleanup
	#TESTS: compiles

Change 2845375 on 2016/01/27 by [REDACTED_PERSON]

	Sequencer: Implement cinematic shot track thumbnails.

	#tests Rebuild the trailer with the cinematic shot track

Change 2845359 on 2016/01/27 by [REDACTED_PERSON]

	Downgrade some checks to ensures.
	#test ps4

Change 2845347 on 2016/01/27 by [REDACTED_PERSON]

	Remove unused EditorStyle dependency from Social. It is not being used, and causes issues for the engine team.
	#TESTS n/a

Change 2845227 on 2016/01/27 by [REDACTED_PERSON]

	Adding flags to create callstack map files when building Arxan protection

	#tests Built arxan exe

Change 2844871 on 2016/01/26 by [REDACTED_PERSON]

	Prevent enums from being regenerated while cooking (prevents false-positive warning about FText's being regenerated)
	#tests ran editor

diff --git a/Engine/Config/BaseEngine.ini b/Engine/Config/BaseEngine.ini
--- a/Engine/Config/BaseEngine.ini
+++ b/Engine/Config/BaseEngine.ini
@@ -1494,6 +1494,7 @@ ZeroEngineVersionWarning=True
 UseStrictEngineVersioning=True
 UseVTableConstructors=True
 CanStripEditorOnlyExportsAndImports=True
+CanSkipEditorReferencedPackagesWhenCooking=True
 DetailedCallstacksInNonMonolithicBuilds=False
 UseSeperateBulkDataFiles=False
 
@@ -1577,7 +1578,7 @@ MaxLightmapRadius=10000.0
 AllowStreamingLightmaps=True
 UseDynamicStreaming=True
 BoostPlayerTextures=3.0
-RangePrefetchDistance=400
+RangePrefetchDistance=1000
 
 
 [/Script/UnrealEd.EditorEngine]
@@ -1874,6 +1875,9 @@ FLightPrimitiveInteractionInitialBlockSize=512
 [SystemSettings]
 ; NOTE THAT ANY ITEMS IN THIS SECTION WILL AFFECT ALL PLATFORMS!!!
 
+// disabled for now to prevent issues http://crashreport/Crashes/Show/2948912, http://crashreport/Crashes/Show/2948842
+r.DetectAndWarnOfBadDrivers = 0
+
 [SystemSettingsEditor]
 ; System settings overrides for the editor.  Ideally the editor should use the same settings as the game.
 
diff --git a/Engine/Plugins/Runtime/PacketHandlers/CompressionComponents/Oodle/Oodle.uplugin b/Engine/Plugins/Runtime/PacketHandlers/CompressionComponents/Oodle/Oodle.uplugin
--- a/Engine/Plugins/Runtime/PacketHandlers/CompressionComponents/Oodle/Oodle.uplugin
+++ b/Engine/Plugins/Runtime/PacketHandlers/CompressionComponents/Oodle/Oodle.uplugin
@@ -24,7 +24,8 @@
 			[
 				"Win32",
 				"Win64",
-				"Linux"
+				"Linux",
+				"PS4"
 			]
 		}
 	]
diff --git a/Engine/Plugins/Runtime/PacketHandlers/CompressionComponents/Oodle/Source/OodleHandlerComponent/Classes/OodleTrainerCommandlet.h b/Engine/Plugins/Runtime/PacketHandlers/CompressionComponents/Oodle/Source/OodleHandlerComponent/Classes/OodleTrainerCommandlet.h
--- a/Engine/Plugins/Runtime/PacketHandlers/CompressionComponents/Oodle/Source/OodleHandlerComponent/Classes/OodleTrainerCommandlet.h
+++ b/Engine/Plugins/Runtime/PacketHandlers/CompressionComponents/Oodle/Source/OodleHandlerComponent/Classes/OodleTrainerCommandlet.h
@@ -64,6 +64,7 @@ public:
 	 */
 	static bool HandleEnable();
 
+#if !UE_BUILD_SHIPPING || OODLE_DEV_SHIPPING
 	/**
 	 * Handles the 'MergePackets' command, which is used to merge multiple packet capture files
 	 *
@@ -109,11 +110,12 @@ public:
 	 * @return				Whether or not OutputFile is a valid path
 	 */
 	static bool VerifyOutputFile(FString OutputFile);
-#endif
+#endif // HAS_OODLE_SDK
+#endif // !UE_BUILD_SHIPPING || OODLE_DEV_SHIPPING
 };
 
 
-#if HAS_OODLE_SDK
+#if HAS_OODLE_SDK && (!UE_BUILD_SHIPPING || OODLE_DEV_SHIPPING)
 /**
  * FOodleDictionaryGenerator
  *
diff --git a/Engine/Plugins/Runtime/PacketHandlers/CompressionComponents/Oodle/Source/OodleHandlerComponent/OodleHandlerComponent.Build.cs b/Engine/Plugins/Runtime/PacketHandlers/CompressionComponents/Oodle/Source/OodleHandlerComponent/OodleHandlerComponent.Build.cs
--- a/Engine/Plugins/Runtime/PacketHandlers/CompressionComponents/Oodle/Source/OodleHandlerComponent/OodleHandlerComponent.Build.cs
+++ b/Engine/Plugins/Runtime/PacketHandlers/CompressionComponents/Oodle/Source/OodleHandlerComponent/OodleHandlerComponent.Build.cs
@@ -29,13 +29,19 @@ public class OodleHandlerComponent : ModuleRules
         {
 			OodleNotForLicenseesLibDir = System.IO.Path.Combine( UEBuildConfiguration.UEThirdPartySourceDirectory, "..", "..",
 				"Plugins", "Runtime", "PacketHandlers", "CompressionComponents", "Oodle", "Source", "ThirdParty", "NotForLicensees",
-				"Oodle", "210beta", "win", "lib" );
+				"Oodle", "212", "win", "lib" );
 		}
 		else if (Target.Platform == UnrealTargetPlatform.Linux)
 		{
 			OodleNotForLicenseesLibDir = System.IO.Path.Combine( UEBuildConfiguration.UEThirdPartySourceDirectory, "..", "..",
 				"Plugins", "Runtime", "PacketHandlers", "CompressionComponents", "Oodle", "Source", "ThirdParty", "NotForLicensees",
-				"Oodle", "210beta", "linux", "lib" );
+				"Oodle", "212", "linux", "lib" );
+		}
+		else if ( Target.Platform == UnrealTargetPlatform.PS4 )
+		{
+			OodleNotForLicenseesLibDir = System.IO.Path.Combine( UEBuildConfiguration.UEThirdPartySourceDirectory, "..", "..",
+				"Plugins", "Runtime", "PacketHandlers", "CompressionComponents", "Oodle", "Source", "ThirdParty", "NotForLicensees",
+				"Oodle", "212", "ps4", "lib" );
 		}
 
 		if (OodleNotForLicenseesLibDir.Length > 0)
diff --git a/Engine/Plugins/Runtime/PacketHandlers/CompressionComponents/Oodle/Source/OodleHandlerComponent/Private/OodleArchives.cpp b/Engine/Plugins/Runtime/PacketHandlers/CompressionComponents/Oodle/Source/OodleHandlerComponent/Private/OodleArchives.cpp
--- a/Engine/Plugins/Runtime/PacketHandlers/CompressionComponents/Oodle/Source/OodleHandlerComponent/Private/OodleArchives.cpp
+++ b/Engine/Plugins/Runtime/PacketHandlers/CompressionComponents/Oodle/Source/OodleHandlerComponent/Private/OodleArchives.cpp
@@ -55,7 +55,7 @@ bool FOodleArchiveBase::SerializeOodleCompressData(FOodleCompressedData& OutData
 			uint32 OffsetPos = InnerArchive.Tell();
 			OutDataInfo.Offset.Set(*this, OffsetPos);
 
-			InnerArchive.Serialize(CompressBuffer, OodleLen);
+			InnerArchive.Serialize((void*)CompressBuffer, OodleLen);
 		}
 
 
@@ -87,7 +87,7 @@ bool FOodleArchiveBase::SerializeOodleDecompressData(FOodleCompressedData& DataI
 
 		InnerArchive.Serialize(CompressedData, CompressedLength);
 
-		SINTa OodleLen = OodleLZ_Decompress(CompressedData, CompressedLength, DecompressedData, DecompressedLength);
+		SINTa OodleLen = OodleLZ_Decompress((void*)CompressedData, CompressedLength, (void*)DecompressedData, DecompressedLength);
 
 		bSuccess = OodleLen != 0;
 
@@ -112,6 +112,7 @@ bool FOodleArchiveBase::SerializeOodleDecompressData(FOodleCompressedData& DataI
 }
 
 
+#if !UE_BUILD_SHIPPING || OODLE_DEV_SHIPPING
 /**
  * FCaptureHeader
  */
@@ -303,6 +304,7 @@ uint32 FPacketCaptureArchive::GetPacketCount()
 
 	return ReturnVal;
 }
+#endif
 
 
 /**
diff --git a/Engine/Plugins/Runtime/PacketHandlers/CompressionComponents/Oodle/Source/OodleHandlerComponent/Private/OodleHandlerComponent.cpp b/Engine/Plugins/Runtime/PacketHandlers/CompressionComponents/Oodle/Source/OodleHandlerComponent/Private/OodleHandlerComponent.cpp
--- a/Engine/Plugins/Runtime/PacketHandlers/CompressionComponents/Oodle/Source/OodleHandlerComponent/Private/OodleHandlerComponent.cpp
+++ b/Engine/Plugins/Runtime/PacketHandlers/CompressionComponents/Oodle/Source/OodleHandlerComponent/Private/OodleHandlerComponent.cpp
@@ -8,16 +8,29 @@
 DEFINE_LOG_CATEGORY(OodleHandlerComponentLog);
 
 
+// @todo #JohnB: IMPORTANT: Do not add any placeholders for the compression-toggling into the protocol yet. Needs careful planning,
+//					might not be packet-based compression, but bunch-based.
+
+
 // @todo #JohnB: Deprecate the 'training' mode name from this handler eventually, as it is misnamed/confusing
 
 // @todo #JohnB: You need context-sensitive data capturing; e.g. you do NOT want to be capturing voice channel data
 
+
+// @todo #JohnB: Get rid of 'Mode' variable, and enable 'Capturing' and 'Release' mode separately (so they can operate at same time);
+//					rename 'Release' mode too, to something a bit more descriptive.
+
+
+// @todo #JohnB: You're not taking into account, the overhead of sending 'DecompressedLength', in the stats
+
+
 #if HAS_OODLE_SDK
 #define OODLE_INI_SECTION TEXT("OodleHandlerComponent")
 
 
 // @todo #JohnB: Find a better solution than this, for handling the maximum packet size. This is far higher than it needs to be.
-#define MAX_OODLE_PACKET_SIZE 65535
+// Reduced from 65535 to 16383, to minimize Oodle packet overhead from 3 bytes to 2 bytes
+#define MAX_OODLE_PACKET_SIZE 16383
 
 
 #if STATS
@@ -37,6 +50,7 @@ DEFINE_STAT(STAT_Oodle_InRaw);
 DEFINE_STAT(STAT_Oodle_InCompressed);
 DEFINE_STAT(STAT_Oodle_InSavings);
 DEFINE_STAT(STAT_Oodle_InTotalSavings);
+DEFINE_STAT(STAT_Oodle_PacketOverhead);
 
 #if !UE_BUILD_SHIPPING
 DEFINE_STAT(STAT_Oodle_InDecompressTime);
@@ -60,6 +74,10 @@ typedef TMap<FString, TSharedPtr<FOodleDictionary>> FDictionaryMap;
 static FDictionaryMap DictionaryMap;
 
 
+/** Whether or not Oodle is presently force-enabled */
+static bool bOodleForceEnable = false;
+
+
 #if STATS
 /** The global net stats tracker for Oodle */
 static FOodleNetStats GOodleNetStats;
@@ -144,9 +162,9 @@ FOodleDictionary::FOodleDictionary(uint32 InHashTableSize, uint8* InDictionaryDa
 {
 #if STATS
 	INC_DWORD_STAT(STAT_Oodle_DictionaryCount);
-	INC_DWORD_STAT_BY(STAT_Oodle_DictionaryBytes, DictionarySize);
-	INC_DWORD_STAT_BY(STAT_Oodle_SharedBytes, SharedDictionarySize);
-	INC_DWORD_STAT_BY(STAT_Oodle_StateBytes, CompressorStateSize);
+	INC_MEMORY_STAT_BY(STAT_Oodle_DictionaryBytes, DictionarySize);
+	INC_MEMORY_STAT_BY(STAT_Oodle_SharedBytes, SharedDictionarySize);
+	INC_MEMORY_STAT_BY(STAT_Oodle_StateBytes, CompressorStateSize);
 #endif
 }
 
@@ -154,9 +172,9 @@ FOodleDictionary::~FOodleDictionary()
 {
 #if STATS
 	DEC_DWORD_STAT(STAT_Oodle_DictionaryCount);
-	DEC_DWORD_STAT_BY(STAT_Oodle_DictionaryBytes, DictionarySize);
-	DEC_DWORD_STAT_BY(STAT_Oodle_SharedBytes, SharedDictionarySize);
-	DEC_DWORD_STAT_BY(STAT_Oodle_StateBytes, CompressorStateSize);
+	DEC_MEMORY_STAT_BY(STAT_Oodle_DictionaryBytes, DictionarySize);
+	DEC_MEMORY_STAT_BY(STAT_Oodle_SharedBytes, SharedDictionarySize);
+	DEC_MEMORY_STAT_BY(STAT_Oodle_StateBytes, CompressorStateSize);
 #endif
 
 	if (DictionaryData != nullptr)
@@ -181,10 +199,10 @@ FOodleDictionary::~FOodleDictionary()
  */
 
 OodleHandlerComponent::OodleHandlerComponent()
-	: InPacketLog(nullptr)
+	: bEnableOodle(false)
+#if !UE_BUILD_SHIPPING || OODLE_DEV_SHIPPING
+	, InPacketLog(nullptr)
 	, OutPacketLog(nullptr)
-	, bEnableOodle(false)
-#if !UE_BUILD_SHIPPING
 	, bUseDictionaryIfPresent(false)
 #endif
 	, Mode(EOodleHandlerMode::Release)
@@ -196,6 +214,7 @@ OodleHandlerComponent::OodleHandlerComponent()
 
 OodleHandlerComponent::~OodleHandlerComponent()
 {
+#if !UE_BUILD_SHIPPING || OODLE_DEV_SHIPPING
 	if (OutPacketLog != nullptr)
 	{
 		OutPacketLog->Close();
@@ -215,6 +234,7 @@ OodleHandlerComponent::~OodleHandlerComponent()
 
 		InPacketLog = nullptr;
 	}
+#endif
 
 
 	FreeDictionary(ServerDictionary);
@@ -234,12 +254,12 @@ void OodleHandlerComponent::InitFirstRunConfig()
 	{
 		GConfig->SetBool(OODLE_INI_SECTION, TEXT("bEnableOodle"), true, GEngineIni);
 
-#if !UE_BUILD_SHIPPING
+#if !UE_BUILD_SHIPPING || OODLE_DEV_SHIPPING
 		GConfig->SetBool(OODLE_INI_SECTION, TEXT("bUseDictionaryIfPresent"), false, GEngineIni);
+		GConfig->SetString(OODLE_INI_SECTION, TEXT("PacketLogFile"), TEXT("PacketDump"), GEngineIni);
 #endif
 
 		GConfig->SetString(OODLE_INI_SECTION, TEXT("Mode"), TEXT("Capturing"), GEngineIni);
-		GConfig->SetString(OODLE_INI_SECTION, TEXT("PacketLogFile"), TEXT("PacketDump"), GEngineIni);
 
 		GConfig->SetString(OODLE_INI_SECTION, TEXT("ServerDictionary"), TEXT(""), GEngineIni);
 		GConfig->SetString(OODLE_INI_SECTION, TEXT("ClientDictionary"), TEXT(""), GEngineIni);
@@ -252,8 +272,6 @@ void OodleHandlerComponent::Initialize()
 {
 	InitFirstRunConfig();
 
-	const bool bOodleForceEnable = FParse::Param(FCommandLine::Get(), TEXT("Oodle"));
-
 	// Class config variables
 	GConfig->GetBool(OODLE_INI_SECTION, TEXT("bEnableOodle"), bEnableOodle, GEngineIni);
 
@@ -263,7 +281,7 @@ void OodleHandlerComponent::Initialize()
 		bEnableOodle = true;
 	}
 
-#if !UE_BUILD_SHIPPING
+#if !UE_BUILD_SHIPPING || OODLE_DEV_SHIPPING
 	GConfig->GetBool(OODLE_INI_SECTION, TEXT("bUseDictionaryIfPresent"), bUseDictionaryIfPresent, GEngineIni);
 
 	if (!bUseDictionaryIfPresent && bOodleForceEnable)
@@ -271,10 +289,8 @@ void OodleHandlerComponent::Initialize()
 		UE_LOG(OodleHandlerComponentLog, Log, TEXT("Force-enabling 'bUseDictionaryIfPresent', due to -Oodle on commandline."));
 		bUseDictionaryIfPresent = true;
 	}
-#endif
 
 
-	// Mode
 	FString ReadMode;
 	GConfig->GetString(OODLE_INI_SECTION, TEXT("Mode"), ReadMode, GEngineIni);
 
@@ -282,19 +298,15 @@ void OodleHandlerComponent::Initialize()
 	{
 		Mode = EOodleHandlerMode::Capturing;
 	}
-	else if(ReadMode == TEXT("Release"))
-	{
-		Mode = EOodleHandlerMode::Release;
-	}
-	// Default
-	else
+#endif
+	else //if (ReadMode == TEXT("Release"))
 	{
 		Mode = EOodleHandlerMode::Release;
 	}
 
 	if (bEnableOodle)
 	{
-#if !UE_BUILD_SHIPPING
+#if !UE_BUILD_SHIPPING || OODLE_DEV_SHIPPING
 		bool bForceCapturing = FParse::Param(FCommandLine::Get(), TEXT("OodleTraining")) ||
 								FParse::Param(FCommandLine::Get(), TEXT("OodleCapturing"));
 
@@ -323,6 +335,7 @@ void OodleHandlerComponent::Initialize()
 
 		switch(Mode)
 		{
+#if !UE_BUILD_SHIPPING || OODLE_DEV_SHIPPING
 			// @todo #JohnB: Convert this code so that just one capture file is used for all connections, per session
 			//					(could set it up much like the dictionary sharing code)
 		case EOodleHandlerMode::Capturing:
@@ -374,6 +387,7 @@ void OodleHandlerComponent::Initialize()
 
 				break;
 			}
+#endif
 
 		case EOodleHandlerMode::Release:
 			{
@@ -381,7 +395,7 @@ void OodleHandlerComponent::Initialize()
 				FString ClientDictionaryPath;
 				bool bGotDictionaryPath = false;
 
-#if !UE_BUILD_SHIPPING
+#if ( !UE_BUILD_SHIPPING || OODLE_DEV_SHIPPING ) && !PLATFORM_PS4
 				if (bUseDictionaryIfPresent)
 				{
 					bGotDictionaryPath = FindFallbackDictionaries(ServerDictionaryPath, ClientDictionaryPath);
@@ -435,7 +449,6 @@ void OodleHandlerComponent::InitializeDictionary(FString FilePath, TSharedPtr<FO
 			BoundArc.SerializeHeader();
 			BoundArc.SerializeDictionaryAndState(DictionaryData, DictionaryBytes, CompactCompressorState, CompactCompressorStateBytes);
 
-
 			if (!BoundArc.IsError())
 			{
 				UE_LOG(OodleHandlerComponentLog, Log, TEXT("Loading dictionary file: %s"), *FilePath);
@@ -452,7 +465,7 @@ void OodleHandlerComponent::InitializeDictionary(FString FilePath, TSharedPtr<FO
 				uint32 SharedDictionarySize = OodleNetwork1_Shared_Size(HashTableSize);
 				OodleNetwork1_Shared* SharedDictionary = (OodleNetwork1_Shared*)FMemory::Malloc(SharedDictionarySize);
 
-				OodleNetwork1_Shared_SetWindow(SharedDictionary, HashTableSize, DictionaryData, DictionaryBytes);
+				OodleNetwork1_Shared_SetWindow(SharedDictionary, HashTableSize, (void*)DictionaryData, DictionaryBytes);
 
 
 				// Now add the dictionary data to the map
@@ -593,7 +606,7 @@ bool OodleHandlerComponent::GetDictionaryPaths(FString& OutServerDictionary, FSt
 	return bSuccess;
 }
 
-#if !UE_BUILD_SHIPPING
+#if !UE_BUILD_SHIPPING || OODLE_DEV_SHIPPING
 bool OodleHandlerComponent::FindFallbackDictionaries(FString& OutServerDictionary, FString& OutClientDictionary, bool bTestOnly/*=false*/)
 {
 	bool bSuccess = false;
@@ -694,6 +707,7 @@ void OodleHandlerComponent::Incoming(FBitReader& Packet)
 	{
 		switch(Mode)
 		{
+#if !UE_BUILD_SHIPPING || OODLE_DEV_SHIPPING
 		case EOodleHandlerMode::Capturing:
 			{
 				if (Handler->Mode == Handler::Mode::Server)
@@ -705,6 +719,7 @@ void OodleHandlerComponent::Incoming(FBitReader& Packet)
 				}
 				break;
 			}
+#endif
 
 		case EOodleHandlerMode::Release:
 			{
@@ -793,6 +808,7 @@ void OodleHandlerComponent::Outgoing(FBitWriter& Packet)
 	{
 		switch(Mode)
 		{
+#if !UE_BUILD_SHIPPING || OODLE_DEV_SHIPPING
 		case EOodleHandlerMode::Capturing:
 			{
 				if (Handler->Mode == Handler::Mode::Server)
@@ -803,6 +819,7 @@ void OodleHandlerComponent::Outgoing(FBitWriter& Packet)
 				}
 				break;
 			}
+#endif
 
 		case EOodleHandlerMode::Release:
 			{
@@ -857,6 +874,93 @@ void OodleHandlerComponent::Outgoing(FBitWriter& Packet)
 	}
 }
 
+int32 OodleHandlerComponent::GetPacketOverheadBits()
+{
+	int32 ReturnVal = 0;
+
+	if (bEnableOodle && Mode == EOodleHandlerMode::Release)
+	{
+		// Oodle writes the decompressed packet size, as its addition to the protocol - it writes this using SerializeIntPacked however,
+		// so determine the worst case number of packed bytes that will be written, based on the MAX_OODLE_PACKET_SIZE packet limit.
+		// Normally, this should be 2-bytes/16-bits.
+		TArray<uint8> MeasureArray;
+		FMemoryWriter MeasureAr(MeasureArray);
+		uint32 MaxOodlePacket = MAX_OODLE_PACKET_SIZE;
+
+		MeasureAr.SerializeIntPacked(MaxOodlePacket);
+
+		if (!MeasureAr.IsError())
+		{
+			ReturnVal += MeasureAr.TotalSize() * 8;
+
+			SET_DWORD_STAT(STAT_Oodle_PacketOverhead, ReturnVal);
+		}
+		else
+		{
+			LowLevelFatalError(TEXT("Failed to determine OodleHandlerComponent packet overhead."));
+		}
+	}
+
+	return ReturnVal;
+}
+
+
+#if !UE_BUILD_SHIPPING
+/**
+ * Exec Interface
+ */
+static bool OodleExec(UWorld* InWorld, const TCHAR* Cmd, FOutputDevice& Ar)
+{
+	bool bReturnVal = false;
+
+	if (FParse::Command(&Cmd, TEXT("Oodle")))
+	{
+		// Used by unit testing code, to enable/disable Oodle during a unit test
+		if (FParse::Command(&Cmd, TEXT("ForceEnable")))
+		{
+			bool bTurnOn = false;
+
+			if (FParse::Command(&Cmd, TEXT("On")))
+			{
+				bTurnOn = true;
+			}
+			else if (FParse::Command(&Cmd, TEXT("Off")))
+			{
+				bTurnOn = false;
+			}
+			else if (FParse::Command(&Cmd, TEXT("Default")))
+			{
+				bTurnOn = FParse::Param(FCommandLine::Get(), TEXT("Oodle"));
+			}
+			else
+			{
+				bTurnOn = !bOodleForceEnable;
+			}
+
+			if (bTurnOn != bOodleForceEnable)
+			{
+				bOodleForceEnable = bTurnOn;
+
+				if (bOodleForceEnable)
+				{
+					UOodleTrainerCommandlet::HandleEnable();
+				}
+			}
+		}
+		else
+		{
+			Ar.Logf(TEXT("Unknown Oodle command 'Oodle %s'"), Cmd);
+		}
+
+		bReturnVal = true;
+	}
+
+	return bReturnVal;
+}
+
+FStaticSelfRegisteringExec OodleExecRegistration(&OodleExec);
+#endif
+
 
 /**
  * Module Interface
@@ -870,7 +974,9 @@ TSharedPtr<HandlerComponent> FOodleComponentModuleInterface::CreateComponentInst
 void FOodleComponentModuleInterface::StartupModule()
 {
 	// If Oodle is force-enabled on the commandline, execute the commandlet-enable command, which also adds to the PacketHandler list
-	if (FParse::Param(FCommandLine::Get(), TEXT("Oodle")))
+	bOodleForceEnable = FParse::Param(FCommandLine::Get(), TEXT("Oodle"));
+
+	if (bOodleForceEnable)
 	{
 		UOodleTrainerCommandlet::HandleEnable();
 	}
@@ -880,17 +986,11 @@ void FOodleComponentModuleInterface::StartupModule()
 	GOodleSaveDir = FPaths::ConvertRelativePathToFull(FPaths::Combine(*FPaths::GameSavedDir(), TEXT("Oodle")));
 	GOodleContentDir = FPaths::ConvertRelativePathToFull(FPaths::Combine(*FPaths::GameContentDir(), TEXT("Oodle")));
 
-	// Load the Oodle library (NOTE: Path mirrored in Oodle.Build.cs)
-	FString OodleBinaryPath = FPaths::EngineDir() / TEXT("Binaries/ThirdParty/NotForLicensees/Oodle/");
-	FString OodleBinaryFile;
-
 #if PLATFORM_WINDOWS
 	{
-	#if UE4_OODLE_VER >= 200
-		OodleBinaryFile = TEXT("oo2core_1");
-	#else
-		OodleBinaryFile = FString::Printf(TEXT("oodle_%i"), UE4_OODLE_VER);
-	#endif
+		// Load the Oodle library (NOTE: Path and fallback path mirrored in Oodle.Build.cs)
+		FString OodleBinaryPath = FPaths::GameDir() / TEXT( "Binaries/ThirdParty/Oodle/" );
+		FString OodleBinaryFile = TEXT( "oo2core_1" );
 
 	#if PLATFORM_64BITS
 		OodleBinaryPath += TEXT("Win64/");
@@ -900,20 +1000,19 @@ void FOodleComponentModuleInterface::StartupModule()
 		OodleBinaryFile += TEXT("_win32.dll");
 	#endif
 
-
 		FPlatformProcess::PushDllDirectory(*OodleBinaryPath);
 
 		OodleDllHandle = FPlatformProcess::GetDllHandle(*(OodleBinaryPath + OodleBinaryFile));
 
 		FPlatformProcess::PopDllDirectory(*OodleBinaryPath);
-	}
-#endif
 
-#if !PLATFORM_LINUX
-	check(OodleDllHandle != nullptr);
+		if ( OodleDllHandle == nullptr )
+		{
+			UE_LOG( OodleHandlerComponentLog, Fatal, TEXT( "Could not find Oodle .dll's in path: %s" ), *( OodleBinaryPath + OodleBinaryFile ) );
+		}
+	}
 #endif
 
-
 #if UE4_OODLE_VER < 200
 	OodleInitOptions Options;
 
@@ -951,4 +1050,4 @@ void FOodleComponentModuleInterface::ShutdownModule()
 }
 #endif
 
-IMPLEMENT_MODULE( FOodleComponentModuleInterface, OodleHandlerComponent );
\ No newline at end of file
+IMPLEMENT_MODULE( FOodleComponentModuleInterface, OodleHandlerComponent );
diff --git a/Engine/Plugins/Runtime/PacketHandlers/CompressionComponents/Oodle/Source/OodleHandlerComponent/Private/OodleTrainerCommandlet.cpp b/Engine/Plugins/Runtime/PacketHandlers/CompressionComponents/Oodle/Source/OodleHandlerComponent/Private/OodleTrainerCommandlet.cpp
--- a/Engine/Plugins/Runtime/PacketHandlers/CompressionComponents/Oodle/Source/OodleHandlerComponent/Private/OodleTrainerCommandlet.cpp
+++ b/Engine/Plugins/Runtime/PacketHandlers/CompressionComponents/Oodle/Source/OodleHandlerComponent/Private/OodleTrainerCommandlet.cpp
@@ -37,6 +37,7 @@ int32 UOodleTrainerCommandlet::Main(const FString& Params)
 		{
 			HandleEnable();
 		}
+#if !UE_BUILD_SHIPPING || OODLE_DEV_SHIPPING
 		else if (MainCmd == TEXT("MergePackets"))
 		{
 			if (Tokens.Num() > 2)
@@ -102,6 +103,7 @@ int32 UOodleTrainerCommandlet::Main(const FString& Params)
 		{
 			HandleAutoGenerateDictionaries();
 		}
+#endif
 		else
 		{
 			UE_LOG(OodleHandlerComponentLog, Error, TEXT("Unrecognized sub-command '%s'"), *MainCmd);
@@ -162,6 +164,7 @@ bool UOodleTrainerCommandlet::HandleEnable()
 	return bSuccess;
 }
 
+#if !UE_BUILD_SHIPPING || OODLE_DEV_SHIPPING
 bool UOodleTrainerCommandlet::HandleMergePackets(FString OutputCapFile, FString MergeList, bool bMergeAll/*=false*/)
 {
 	bool bSuccess = true;
@@ -659,6 +662,7 @@ bool FOodleDictionaryGenerator::ReadPackets(FString InputCaptureFiles, bool bMer
 		TrainerPackets.Empty(AvgSplitSize * 2);
 		TrainerPacketSizes.Empty(AvgSplitSize * 2);
 
+		// @todo #JohnB: Enable from commandline (bool currently always false)
 		if (bCompressionTestPackets)
 		{
 			CompressionTestPackets.Empty(AvgSplitSize);
@@ -754,6 +758,9 @@ bool FOodleDictionaryGenerator::ReadPackets(FString InputCaptureFiles, bool bMer
 
 bool FOodleDictionaryGenerator::GenerateAndWriteDictionary()
 {
+#if PLATFORM_PS4
+	return false;
+#else
 	bool bSuccess = false;
 	uint8* NewDictionaryData = new uint8[DictionarySize];
 
@@ -770,16 +777,15 @@ bool FOodleDictionaryGenerator::GenerateAndWriteDictionary()
 	UE_LOG(OodleHandlerComponentLog, Log, TEXT("- TrainerPackets: %i, Size: %i (~%iMB)"),
 			TrainerPackets.Num(), TrainerPacketBytes, (TrainerPacketBytes / (1024 * 1024)));
 
-
-	if (bNoTrials)
+	if ( bNoTrials )
 	{
-		OodleNetwork1_SelectDictionaryFromPackets(NewDictionaryData, DictionarySize, HashTableSize,
+		OodleNetwork1_SelectDictionaryFromPackets((void*)NewDictionaryData, DictionarySize, HashTableSize,
 			(const void**)DictionaryPackets.GetData(), DictionaryPacketSizes.GetData(), DictionaryPackets.Num(),
 			(const void**)DictionaryTestPackets.GetData(), DictionaryTestPacketSizes.GetData(), DictionaryTestPackets.Num());
 	}
 	else
 	{
-		OodleNetwork1_SelectDictionaryFromPackets_Trials(NewDictionaryData, DictionarySize, HashTableSize,
+		OodleNetwork1_SelectDictionaryFromPackets_Trials((void*)NewDictionaryData, DictionarySize, HashTableSize,
 			(const void**)DictionaryPackets.GetData(), DictionaryPacketSizes.GetData(), DictionaryPackets.Num(),
 			(const void**)DictionaryTestPackets.GetData(), DictionaryTestPacketSizes.GetData(), DictionaryTestPackets.Num(),
 			DictionaryTrials, TrialRandomness, TrialGenerations);
@@ -794,7 +800,7 @@ bool FOodleDictionaryGenerator::GenerateAndWriteDictionary()
 	OodleNetwork1_Shared* SharedDictionary = (OodleNetwork1_Shared*)FMemory::Malloc(SharedStruct_Size);
 	OodleNetwork1UDP_State* CompressorState = (OodleNetwork1UDP_State*)FMemory::Malloc(StateStruct_Size);
 
-	OodleNetwork1_Shared_SetWindow(SharedDictionary, HashTableSize, NewDictionaryData, (S32)DictionarySize);
+	OodleNetwork1_Shared_SetWindow(SharedDictionary, HashTableSize, (void*)NewDictionaryData, (S32)DictionarySize);
 
 	OodleNetwork1UDP_Train(CompressorState, SharedDictionary,
 		(const void**)TrainerPackets.GetData(), TrainerPacketSizes.GetData(), TrainerPackets.Num());
@@ -819,8 +825,7 @@ bool FOodleDictionaryGenerator::GenerateAndWriteDictionary()
 
 		bSuccess = !OutputFile.IsError();
 
-		bSuccess = bSuccess && OutputFile.SerializeOodleCompressData(OutputFile.Header.DictionaryData, NewDictionaryData,
-																		DictionarySize);
+		bSuccess = bSuccess && OutputFile.SerializeOodleCompressData(OutputFile.Header.DictionaryData, NewDictionaryData, DictionarySize);
 
 		bSuccess = bSuccess && OutputFile.SerializeOodleCompressData(OutputFile.Header.CompressorData, (uint8*)CompactCompressorState,
 																		CompactCompressorStateBytes);
@@ -884,6 +889,7 @@ bool FOodleDictionaryGenerator::GenerateAndWriteDictionary()
 
 
 	return bSuccess;
+#endif
 }
 
 void FOodleDictionaryGenerator::Cleanup()
@@ -941,7 +947,9 @@ void FOodleDictionaryGenerator::Cleanup()
 		MergeMap.Empty();
 	}
 }
-#endif
+#endif // HAS_OODLE_SDK
+
+#endif // !UE_BUILD_SHIPPING || OODLE_DEV_SHIPPING
 
 #undef LOCTEXT_NAMESPACE
 
diff --git a/Engine/Plugins/Runtime/PacketHandlers/CompressionComponents/Oodle/Source/OodleHandlerComponent/Public/OodleArchives.h b/Engine/Plugins/Runtime/PacketHandlers/CompressionComponents/Oodle/Source/OodleHandlerComponent/Public/OodleArchives.h
--- a/Engine/Plugins/Runtime/PacketHandlers/CompressionComponents/Oodle/Source/OodleHandlerComponent/Public/OodleArchives.h
+++ b/Engine/Plugins/Runtime/PacketHandlers/CompressionComponents/Oodle/Source/OodleHandlerComponent/Public/OodleArchives.h
@@ -246,6 +246,7 @@ public:
 };
 
 
+#if !UE_BUILD_SHIPPING || OODLE_DEV_SHIPPING
 /**
  * Archive for handling packet capture (.ucap) files
  */
@@ -346,6 +347,7 @@ public:
 	 */
 	uint32 GetPacketCount();
 };
+#endif
 
 
 /**
diff --git a/Engine/Plugins/Runtime/PacketHandlers/CompressionComponents/Oodle/Source/OodleHandlerComponent/Public/OodleHandlerComponent.h b/Engine/Plugins/Runtime/PacketHandlers/CompressionComponents/Oodle/Source/OodleHandlerComponent/Public/OodleHandlerComponent.h
--- a/Engine/Plugins/Runtime/PacketHandlers/CompressionComponents/Oodle/Source/OodleHandlerComponent/Public/OodleHandlerComponent.h
+++ b/Engine/Plugins/Runtime/PacketHandlers/CompressionComponents/Oodle/Source/OodleHandlerComponent/Public/OodleHandlerComponent.h
@@ -10,6 +10,15 @@
 DECLARE_LOG_CATEGORY_EXTERN(OodleHandlerComponentLog, Log, All);
 
 
+/**
+ * Whether or not to enable Oodle dev code (packet capturing, dictionary training, and automatic dictionary finding) in shipping mode.
+ *
+ * This may be useful for multiplayer game mod authors, to optimize netcode compression for their mod (not officially supported).
+ * However, Oodle compression makes the games network protocol harder to reverse-engineer - enabling this removes that slight benefit.
+ */
+#define OODLE_DEV_SHIPPING	FALSE
+
+
 #if HAS_OODLE_SDK
 #include "OodleArchives.h"
 
@@ -51,6 +60,8 @@ DECLARE_DWORD_ACCUMULATOR_STAT_EXTERN(TEXT("Oodle In Rate Compressed (bytes)"),
 DECLARE_FLOAT_ACCUMULATOR_STAT_EXTERN(TEXT("Oodle In Rate Savings %"), STAT_Oodle_InSavings, STATGROUP_Oodle, );
 DECLARE_FLOAT_ACCUMULATOR_STAT_EXTERN(TEXT("Oodle In Total Savings %"), STAT_Oodle_InTotalSavings, STATGROUP_Oodle, );
 
+DECLARE_DWORD_ACCUMULATOR_STAT_EXTERN(TEXT("Oodle Packet Overhead (bits)"), STAT_Oodle_PacketOverhead, STATGROUP_Oodle, );
+
 #if !UE_BUILD_SHIPPING
 DECLARE_CYCLE_STAT_EXTERN(TEXT("Oodle Out Compress Time"), STAT_Oodle_OutCompressTime, STATGROUP_Oodle, );
 DECLARE_CYCLE_STAT_EXTERN(TEXT("Oodle In Decompress Time"), STAT_Oodle_InDecompressTime, STATGROUP_Oodle, );
@@ -291,7 +302,7 @@ public:
 	 */
 	bool GetDictionaryPaths(FString& OutServerDictionary, FString& OutClientDictionary, bool bFailFatal=true);
 
-#if !UE_BUILD_SHIPPING
+#if !UE_BUILD_SHIPPING || OODLE_DEV_SHIPPING
 	/**
 	 * Searches the game directory for alternate/fallback dictionary files, using the *.udic file extension.
 	 * NOTE: This is non-shipping-only, as release games MUST have well-determined dictionary files (for net-binary-compatibility)
@@ -316,18 +327,19 @@ public:
 	/* Handles any outgoing packets */
 	virtual void Outgoing(FBitWriter& Packet) override;
 
+	virtual int32 GetPacketOverheadBits() override;
+
 protected:
+	/** Whether or not Oodle is enabled */
+	bool bEnableOodle;
+
+#if !UE_BUILD_SHIPPING || OODLE_DEV_SHIPPING
 	/* File to log input packets to */
 	FPacketCaptureArchive* InPacketLog;
 
 	/* File to log output packets to */
 	FPacketCaptureArchive* OutPacketLog;
 
-
-	/** Whether or not Oodle is enabled */
-	bool bEnableOodle;
-
-#if !UE_BUILD_SHIPPING
 	/** Search for dictionary files and use them if present - switching mode to Release in process - don't use in shipping */
 	bool bUseDictionaryIfPresent;
 #endif
diff --git a/Engine/Shaders/MaterialTemplate.usf b/Engine/Shaders/MaterialTemplate.usf
--- a/Engine/Shaders/MaterialTemplate.usf
+++ b/Engine/Shaders/MaterialTemplate.usf
@@ -470,7 +470,7 @@ float4 GetScreenPosition(FMaterialPixelParameters Parameters)
 	}
 #endif //DECAL_RENDERSTAGE
 
-#if DECAL_RENDERSTAGE && DECAL_PRIMITIVE
+#if DECAL_PRIMITIVE
 	float DecalLifetimeOpacity()
 	{
 		return DecalParams.y;
@@ -480,7 +480,7 @@ float4 GetScreenPosition(FMaterialPixelParameters Parameters)
 	{
 		return 0.0f;
 	}
-#endif // DECAL_RENDERSTAGE && DECAL_PRIMITIVE
+#endif // DECAL_PRIMITIVE
 
 /** Transforms a vector from tangent space to world space, prescaling by an amount calculated previously */
 MaterialFloat3 TransformTangentVectorToWorld_PreScaled(FMaterialTessellationParameters Parameters, MaterialFloat3 InTangentVector)
diff --git a/Engine/Source/Developer/AssetTools/Private/AssetTypeActions/AssetTypeActions_CameraAnim.cpp b/Engine/Source/Developer/AssetTools/Private/AssetTypeActions/AssetTypeActions_CameraAnim.cpp
--- a/Engine/Source/Developer/AssetTools/Private/AssetTypeActions/AssetTypeActions_CameraAnim.cpp
+++ b/Engine/Source/Developer/AssetTools/Private/AssetTypeActions/AssetTypeActions_CameraAnim.cpp
@@ -12,8 +12,6 @@
 #include "Camera/CameraActor.h"
 #include "GameFramework/Pawn.h"
 
-DEFINE_LOG_CATEGORY_STATIC(LogCameraAnim, Log, All);
-
 #define LOCTEXT_NAMESPACE "AssetTypeActions"
 
 void FAssetTypeActions_CameraAnim::CreateMatineeActorForCameraAnim(UCameraAnim* InCameraAnim)
diff --git a/Engine/Source/Developer/FriendsAndChat/FriendsAndChat.Build.cs b/Engine/Source/Developer/FriendsAndChat/FriendsAndChat.Build.cs
--- a/Engine/Source/Developer/FriendsAndChat/FriendsAndChat.Build.cs
+++ b/Engine/Source/Developer/FriendsAndChat/FriendsAndChat.Build.cs
@@ -35,17 +35,5 @@ public class FriendsAndChat : ModuleRules
 				"Developer/FriendsAndChat/Public",
 			}
 		);
-
-		DynamicallyLoadedModuleNames.AddRange(
-		new string[] {
-				"EditorStyle",
-			}
-		);
-
-		PrivateIncludePathModuleNames.AddRange(
-			new string[] {
-				"EditorStyle",
-			}
-		);
 	}
 }
diff --git a/Engine/Source/Developer/SynthBenchmark/Private/SynthBenchmarkPrivatePCH.cpp b/Engine/Source/Developer/SynthBenchmark/Private/SynthBenchmarkPrivatePCH.cpp
--- a/Engine/Source/Developer/SynthBenchmark/Private/SynthBenchmarkPrivatePCH.cpp
+++ b/Engine/Source/Developer/SynthBenchmark/Private/SynthBenchmarkPrivatePCH.cpp
@@ -126,6 +126,7 @@ void FSynthBenchmark::Run(FSynthBenchmarkResults& InOut, bool bGPUBenchmark, flo
 	UE_LOG(LogSynthBenchmark, Display, TEXT("  Adapter Name: '%s'"), *GRHIAdapterName);
 	UE_LOG(LogSynthBenchmark, Display, TEXT("  (On Optimus the name might be wrong, memory should be ok)"));
 	UE_LOG(LogSynthBenchmark, Display, TEXT("  Vendor Id: 0x%x"), GRHIVendorId);
+	UE_LOG(LogSynthBenchmark, Display, TEXT("  Device Id: 0x%x"), GRHIDeviceId);
 
 	{
 		FTextureMemoryStats Stats;
diff --git a/Engine/Source/Editor/GameplayAbilitiesEditor/Private/AttributeDetails.cpp b/Engine/Source/Editor/GameplayAbilitiesEditor/Private/AttributeDetails.cpp
--- a/Engine/Source/Editor/GameplayAbilitiesEditor/Private/AttributeDetails.cpp
+++ b/Engine/Source/Editor/GameplayAbilitiesEditor/Private/AttributeDetails.cpp
@@ -303,10 +303,28 @@ void FScalableFloatDetails::CustomizeHeader( TSharedRef<class IPropertyHandle> S
 	StructPropertyHandle->GetNumChildren(NumChildren);
 
 	ValueProperty = StructPropertyHandle->GetChildHandle(GET_MEMBER_NAME_CHECKED(FScalableFloat,Value));
+	if (ValueProperty.IsValid() == false)
+	{
+		return;
+	}
+
 	CurveTableHandleProperty = StructPropertyHandle->GetChildHandle(GET_MEMBER_NAME_CHECKED(FScalableFloat,Curve));
+	if (CurveTableHandleProperty.IsValid() == false)
+	{
+		return;
+	}
 
 	RowNameProperty = CurveTableHandleProperty->GetChildHandle(GET_MEMBER_NAME_CHECKED(FCurveTableRowHandle, RowName));
+	if (RowNameProperty.IsValid() == false)
+	{
+		return;
+	}
+
 	CurveTableProperty = CurveTableHandleProperty->GetChildHandle(GET_MEMBER_NAME_CHECKED(FCurveTableRowHandle, CurveTable));
+	if (CurveTableProperty.IsValid() == false)
+	{
+		return;
+	}
 
 	CurrentSelectedItem = InitWidgetContent();
 
diff --git a/Engine/Source/Editor/GameplayAbilitiesEditor/Private/SGameplayCueEditor.cpp b/Engine/Source/Editor/GameplayAbilitiesEditor/Private/SGameplayCueEditor.cpp
--- a/Engine/Source/Editor/GameplayAbilitiesEditor/Private/SGameplayCueEditor.cpp
+++ b/Engine/Source/Editor/GameplayAbilitiesEditor/Private/SGameplayCueEditor.cpp
@@ -695,14 +695,21 @@ private:
 		FString PerfMessage = FString::Printf(TEXT("SGameplayCueEditor::UpdateGameplayCueListItems"));
 		SCOPE_LOG_TIME_IN_SECONDS(*PerfMessage, nullptr)
 #endif
+		UGameplayCueManager* CueManager = UAbilitySystemGlobals::Get().GetGameplayCueManager();
+		if (!CueManager)
+		{
+			return;
+		}
 
 		GameplayCueListItems.Reset();
 		IGameplayTagsModule& GameplayTagModule = IGameplayTagsModule::Get();
-		UGameplayCueManager* CueManager = UAbilitySystemGlobals::Get().GetGameplayCueManager();
+		
 		TSharedPtr<FCueItem> SelectItem;
 
 		FString SearchString = SearchText.ToString();
 
+		CueManager->LoadAllGameplayCueNotifiesForEditor();
+
 		FGameplayTagContainer AllGameplayCueTags = IGameplayTagsModule::Get().GetGameplayTagsManager().RequestGameplayTagChildren(UGameplayCueSet::BaseGameplayCueTag());
 		for (FGameplayTag ThisGameplayCueTag : AllGameplayCueTags)
 		{
diff --git a/Engine/Source/Editor/Persona/Private/AnimationEditorViewportClient.cpp b/Engine/Source/Editor/Persona/Private/AnimationEditorViewportClient.cpp
--- a/Engine/Source/Editor/Persona/Private/AnimationEditorViewportClient.cpp
+++ b/Engine/Source/Editor/Persona/Private/AnimationEditorViewportClient.cpp
@@ -350,7 +350,6 @@ void FAnimationViewportClient::SetPreviewMeshComponent(UDebugSkelMeshComponent*
 
 		// Set to block all to enable tracing.
 		PreviewSkelMeshComp->SetCollisionProfileName(UCollisionProfile::BlockAll_ProfileName);
-		PreviewSkelMeshComp->RefreshBoneTransforms();
 	}
 }
 
diff --git a/Engine/Source/Editor/UnrealEd/Private/Commandlets/ChunkManifestGenerator.cpp b/Engine/Source/Editor/UnrealEd/Private/Commandlets/ChunkManifestGenerator.cpp
--- a/Engine/Source/Editor/UnrealEd/Private/Commandlets/ChunkManifestGenerator.cpp
+++ b/Engine/Source/Editor/UnrealEd/Private/Commandlets/ChunkManifestGenerator.cpp
@@ -1172,6 +1172,8 @@ bool FChunkManifestGenerator::CheckChunkAssetsAreNotInChild(const FChunkDependen
 
 void FChunkManifestGenerator::AddPackageAndDependenciesToChunk(FChunkPackageSet* ThisPackageSet, FName InPkgName, const FString& InSandboxFile, int32 ChunkID, FSandboxPlatformFile* SandboxPlatformFile)
 {
+	FChunkPackageSet* InitialPackageSetForThisChunk = ChunkManifests[ChunkID];
+
 	//Add this asset
 	ThisPackageSet->Add(InPkgName, InSandboxFile);
 
@@ -1203,14 +1205,22 @@ void FChunkManifestGenerator::AddPackageAndDependenciesToChunk(FChunkPackageSet*
 				FString DependentSandboxFile = SandboxPlatformFile->ConvertToAbsolutePathForExternalAppForWrite(*FPackageName::LongPackageNameToFilename(*FilteredPackageName.ToString()));
 				if (!ThisPackageSet->Contains(FilteredPackageName))
 				{
-					UE_LOG(LogChunkManifestGenerator, Log, TEXT("Adding %s to chunk %i because %s depends on it."), *FilteredPackageName.ToString(), ChunkID, *InPkgName.ToString());
-
-					TArray<FName> VisitedPackages;
-					TArray<FName> DependencyChain;
-					GetPackageDependencyChain(InPkgName, PkgName, VisitedPackages, DependencyChain);					
-					for (const auto& ChainName : DependencyChain)
+					if ((InitialPackageSetForThisChunk != nullptr) && InitialPackageSetForThisChunk->Contains(PkgName))
+					{
+						// Don't print anything out; it was pre-assigned to this chunk but we haven't gotten to it yet in the calling loop; we'll go ahead and grab it now
+					}
+					else
 					{
-						UE_LOG(LogChunkManifestGenerator, Log, TEXT("\tchain: %s"), *ChainName.ToString());
+						// It was not assigned to this chunk and we're forcing it to be dragged in, let the user known
+						UE_LOG(LogChunkManifestGenerator, Log, TEXT("Adding %s to chunk %i because %s depends on it."), *FilteredPackageName.ToString(), ChunkID, *InPkgName.ToString());
+
+						TArray<FName> VisitedPackages;
+						TArray<FName> DependencyChain;
+						GetPackageDependencyChain(InPkgName, PkgName, VisitedPackages, DependencyChain);
+						for (const auto& ChainName : DependencyChain)
+						{
+							UE_LOG(LogChunkManifestGenerator, Log, TEXT("\tchain: %s"), *ChainName.ToString());
+						}
 					}
 				}
 				ThisPackageSet->Add(FilteredPackageName, DependentSandboxFile);
diff --git a/Engine/Source/Editor/UnrealEd/Private/Commandlets/ContentCommandlets.cpp b/Engine/Source/Editor/UnrealEd/Private/Commandlets/ContentCommandlets.cpp
--- a/Engine/Source/Editor/UnrealEd/Private/Commandlets/ContentCommandlets.cpp
+++ b/Engine/Source/Editor/UnrealEd/Private/Commandlets/ContentCommandlets.cpp
@@ -124,7 +124,7 @@ int32 UResavePackagesCommandlet::InitializeResaveParameters( const TArray<FStrin
 					FString PackageFile;
 					if (FPackageName::SearchForPackageOnDisk(Line, NULL, &PackageFile, false))
 					{
-						PackageNames.Add(*PackageFile);
+						PackageNames.AddUnique(*PackageFile);
 					}
 					else
 					{
diff --git a/Engine/Source/Editor/UnrealEd/Private/Commandlets/GatherTextFromAssetsCommandlet.cpp b/Engine/Source/Editor/UnrealEd/Private/Commandlets/GatherTextFromAssetsCommandlet.cpp
--- a/Engine/Source/Editor/UnrealEd/Private/Commandlets/GatherTextFromAssetsCommandlet.cpp
+++ b/Engine/Source/Editor/UnrealEd/Private/Commandlets/GatherTextFromAssetsCommandlet.cpp
@@ -13,6 +13,104 @@
 
 DEFINE_LOG_CATEGORY_STATIC(LogGatherTextFromAssetsCommandlet, Log, All);
 
+/** Special feedback context used to stop the commandlet to reporting failure due to a package load error */
+class FLoadPackageLogOutputRedirector : public FFeedbackContext
+{
+public:
+	struct FScopedCapture
+	{
+		FScopedCapture(FLoadPackageLogOutputRedirector* InLogOutputRedirector, const FString& InPackageContext)
+			: LogOutputRedirector(InLogOutputRedirector)
+		{
+			LogOutputRedirector->BeginCapturingLogData(InPackageContext);
+		}
+
+		~FScopedCapture()
+		{
+			LogOutputRedirector->EndCapturingLogData();
+		}
+
+		FLoadPackageLogOutputRedirector* LogOutputRedirector;
+	};
+
+	FLoadPackageLogOutputRedirector()
+		: ErrorCount(0)
+		, WarningCount(0)
+		, FormattedErrorsAndWarningsList()
+		, PackageContext()
+		, OriginalWarningContext(nullptr)
+	{
+	}
+
+	virtual ~FLoadPackageLogOutputRedirector()
+	{
+	}
+
+	void BeginCapturingLogData(const FString& InPackageContext)
+	{
+		// Override GWarn so that we can capture any log data
+		check(!OriginalWarningContext);
+		OriginalWarningContext = GWarn;
+		GWarn = this;
+
+		PackageContext = InPackageContext;
+
+		// Reset the counts and previous log output
+		ErrorCount = 0;
+		WarningCount = 0;
+		FormattedErrorsAndWarningsList.Reset();
+	}
+
+	void EndCapturingLogData()
+	{
+		// Restore the original GWarn now that we've finished capturing log data
+		check(OriginalWarningContext);
+		GWarn = OriginalWarningContext;
+		OriginalWarningContext = nullptr;
+
+		// Report any messages, and also report a warning if we silenced some warnings or errors when loading
+		if (ErrorCount > 0 || WarningCount > 0)
+		{
+			static const FString LogIndentation = TEXT("    ");
+
+			UE_LOG(LogGatherTextFromAssetsCommandlet, Warning, TEXT("Package '%s' produced %d error(s) and %d warning(s) while loading. Please verify that your text has gathered correctly."), *PackageContext, ErrorCount, WarningCount);
+			
+			GWarn->Log(NAME_None, ELogVerbosity::Log, FString::Printf(TEXT("The following errors and warnings were reported while loading '%s':"), *PackageContext));
+			for (const auto& FormattedOutput : FormattedErrorsAndWarningsList)
+			{
+				GWarn->Log(NAME_None, ELogVerbosity::Log, LogIndentation + FormattedOutput);
+			}
+		}
+	}
+
+	virtual void Serialize(const TCHAR* V, ELogVerbosity::Type Verbosity, const FName& Category) override
+	{
+		if (Verbosity == ELogVerbosity::Error)
+		{
+			++ErrorCount;
+			FormattedErrorsAndWarningsList.Add(FormatLogLine(Verbosity, Category, V));
+		}
+		else if (Verbosity == ELogVerbosity::Warning)
+		{
+			++WarningCount;
+			FormattedErrorsAndWarningsList.Add(FormatLogLine(Verbosity, Category, V));
+		}
+		else
+		{
+			// Pass anything else on to GWarn so that it can handle them appropriately
+			OriginalWarningContext->Serialize(V, Verbosity, Category);
+		}
+	}
+
+private:
+	int32 ErrorCount;
+	int32 WarningCount;
+	TArray<FString> FormattedErrorsAndWarningsList;
+
+	FString PackageContext;
+	FFeedbackContext* OriginalWarningContext;
+};
+
 #define LOC_DEFINE_REGION
 
 //////////////////////////////////////////////////////////////////////////
@@ -423,6 +521,8 @@ int32 UGatherTextFromAssetsCommandlet::Main(const FString& Params)
 		UE_LOG(LogGatherTextFromAssetsCommandlet, Log, TEXT("Loading %i packages in %i batches of %i."), PackageCount, BatchCount, PackagesPerBatchCount);
 	}
 
+	FLoadPackageLogOutputRedirector LogOutputRedirector;
+
 	//Load the packages in batches
 	int32 PackageIndex = 0;
 	for( int32 BatchIndex = 0; BatchIndex < BatchCount; ++BatchIndex )
@@ -435,7 +535,18 @@ int32 UGatherTextFromAssetsCommandlet::Main(const FString& Params)
 
 			UE_LOG(LogGatherTextFromAssetsCommandlet, Verbose, TEXT("Loading package: '%s'."), *PackageFileName);
 
-			UPackage *Package = LoadPackage( NULL, *PackageFileName, LOAD_None );
+			UPackage *Package = nullptr;
+			{
+				FString LongPackageName;
+				if (!FPackageName::TryConvertFilenameToLongPackageName(PackageFileName, LongPackageName))
+				{
+					LongPackageName = FPaths::GetCleanFilename(PackageFileName);
+				}
+
+				FLoadPackageLogOutputRedirector::FScopedCapture ScopedCapture(&LogOutputRedirector, LongPackageName);
+				Package = LoadPackage( NULL, *PackageFileName, LOAD_NoWarn | LOAD_Quiet );
+			}
+
 			if( Package )
 			{
 				LoadedPackages.Add(Package);
diff --git a/Engine/Source/Editor/UnrealEd/Private/Tests/AutomationEditorCommon.cpp b/Engine/Source/Editor/UnrealEd/Private/Tests/AutomationEditorCommon.cpp
--- a/Engine/Source/Editor/UnrealEd/Private/Tests/AutomationEditorCommon.cpp
+++ b/Engine/Source/Editor/UnrealEd/Private/Tests/AutomationEditorCommon.cpp
@@ -923,17 +923,25 @@ void FEditorAutomationTestUtilities::CollectGameContentTests(TArray<FString>& Ou
 	for (auto ObjIter = ObjectList.CreateConstIterator(); ObjIter; ++ObjIter)
 	{
 		const FAssetData& Asset = *ObjIter;
-		FString Filename = Asset.ObjectPath.ToString();
-
-		if (Filename.StartsWith("/Game"))
+		if (Asset.GetClass() == nullptr)
 		{
-			//convert to full paths
-			Filename = FPackageName::LongPackageNameToFilename(Filename);
-			if (FAutomationTestFramework::GetInstance().ShouldTestContent(Filename))
+			// a nullptr class is bad !
+			UE_LOG(LogAutomationEditorCommon, Warning, TEXT("GetClass for %s (%s) returned nullptr. Asset ignored"), *Asset.AssetName.ToString(), *Asset.ObjectPath.ToString());
+		}
+		else 
+		{
+			FString Filename = Asset.ObjectPath.ToString();
+
+			if (Filename.StartsWith("/Game"))
 			{
-				FString BeautifiedFilename = FString::Printf(TEXT("%s.%s"), *Asset.GetClass()->GetFName().ToString(), *Asset.AssetName.ToString());
-				OutBeautifiedNames.Add(BeautifiedFilename);
-				OutTestCommands.Add(Asset.ObjectPath.ToString());
+				//convert to full paths
+				Filename = FPackageName::LongPackageNameToFilename(Filename);
+				if (FAutomationTestFramework::GetInstance().ShouldTestContent(Filename))
+				{
+					FString BeautifiedFilename = FString::Printf(TEXT("%s.%s"), *Asset.GetClass()->GetFName().ToString(), *Asset.AssetName.ToString());
+					OutBeautifiedNames.Add(BeautifiedFilename);
+					OutTestCommands.Add(Asset.ObjectPath.ToString());
+				}
 			}
 		}
 	}
diff --git a/Engine/Source/Programs/AutomationTool/AutomationUtils/CommandUtils.cs b/Engine/Source/Programs/AutomationTool/AutomationUtils/CommandUtils.cs
--- a/Engine/Source/Programs/AutomationTool/AutomationUtils/CommandUtils.cs
+++ b/Engine/Source/Programs/AutomationTool/AutomationUtils/CommandUtils.cs
@@ -1412,7 +1412,7 @@ namespace AutomationTool
 			}
 			foreach (var SourceSubDirectory in Directory.GetDirectories(Source))
 			{
-				string DestPath = Dest + GetPathSeparatorChar(PathSeparator.Default) + GetLastDirectoryName(SourceSubDirectory);
+				string DestPath = Dest + GetPathSeparatorChar(PathSeparator.Default) + GetLastDirectoryName(SourceSubDirectory + GetPathSeparatorChar(PathSeparator.Default));
 				if (!CopyDirectory_NoExceptions(SourceSubDirectory, DestPath, bQuiet))
 				{
 					return false;
diff --git a/Engine/Source/Programs/AutomationTool/AutomationUtils/DeploymentContext.cs b/Engine/Source/Programs/AutomationTool/AutomationUtils/DeploymentContext.cs
--- a/Engine/Source/Programs/AutomationTool/AutomationUtils/DeploymentContext.cs
+++ b/Engine/Source/Programs/AutomationTool/AutomationUtils/DeploymentContext.cs
@@ -634,7 +634,7 @@ public class DeploymentContext //: ProjectParams
 		}
 	}
 
-	public int ArchiveFiles(string InPath, string Wildcard = "*", bool bRecursive = true, string[] ExcludeWildcard = null)
+	public int ArchiveFiles(string InPath, string Wildcard = "*", bool bRecursive = true, string[] ExcludeWildcard = null, string NewPath = null)
 	{
 		int FilesAdded = 0;
 
@@ -710,7 +710,22 @@ public class DeploymentContext //: ProjectParams
 				{
 					throw new AutomationException("Can't archive {0}; it was supposed to start with {1}", FileToCopy, InPath);
 				}
-				Dest = FileToCopy.Substring(InPath.Length);
+
+				// If the specified a new directory, first we deal with that, then apply the other things
+				// this is used to collapse the sandbox, among other things
+				if (NewPath != null)
+				{
+					Dest = FileToCopy.Substring(InPath.Length);
+					if (Dest.StartsWith("/") || Dest.StartsWith("\\"))
+					{
+						Dest = Dest.Substring(1);
+					}
+					Dest = CommandUtils.CombinePaths(NewPath, Dest);
+				}
+				else
+				{
+					Dest = FileToCopy.Substring(InPath.Length);
+				}
 
 				if (Dest.StartsWith("/") || Dest.StartsWith("\\"))
 				{
diff --git a/Engine/Source/Programs/AutomationTool/AutomationUtils/Platform.cs b/Engine/Source/Programs/AutomationTool/AutomationUtils/Platform.cs
--- a/Engine/Source/Programs/AutomationTool/AutomationUtils/Platform.cs
+++ b/Engine/Source/Programs/AutomationTool/AutomationUtils/Platform.cs
@@ -452,7 +452,7 @@ namespace AutomationTool
 		/// of current data against a shipped pak file.
 		/// </summary>
 		/// <returns></returns>
-		public virtual bool GetPlatformPatchesWithDiffPak()
+        public virtual bool GetPlatformPatchesWithDiffPak(ProjectParams Params, DeploymentContext SC)
 		{
 			return true;
 		}
diff --git a/Engine/Source/Programs/AutomationTool/GUBP/LegacyNodes.cs b/Engine/Source/Programs/AutomationTool/GUBP/LegacyNodes.cs
--- a/Engine/Source/Programs/AutomationTool/GUBP/LegacyNodes.cs
+++ b/Engine/Source/Programs/AutomationTool/GUBP/LegacyNodes.cs
@@ -2694,9 +2694,25 @@ partial class GUBP
 			// Run the localise script.
 			CommandUtils.RunUAT(CommandUtils.CmdEnv, "Localise" + CommandLineArguments);
 
-			// Don't pass on any build products to other build nodes at the moment.
 			BuildProducts = new List<string>();
-			SaveRecordOfSuccessAndAddToBuildProducts();
+
+			var SuccessMessage = "Updated the following localization build products:\n";
+			try
+			{
+				// Just add all of the localization content as a build product, since UAT should have updated all of it
+				var LocalizationContentDirectory = CombinePaths(CommandUtils.CmdEnv.LocalRoot, UEProjectDirectory, "Content", "Localization");
+				foreach (string FilePath in Directory.EnumerateFiles(LocalizationContentDirectory, "*", SearchOption.AllDirectories))
+				{
+					AddBuildProduct(FilePath);
+					SuccessMessage += "    " + FilePath + "\n";
+				}
+			}
+			catch (Exception e)
+			{
+				LogWarning("Failed to add build products for localization node. The follow exception was raised '%s'", e);
+			}
+
+			SaveRecordOfSuccessAndAddToBuildProducts(SuccessMessage);
 		}
 
 		public override int CISFrequencyQuantumShift(GUBP.GUBPBranchConfig BranchConfig)
diff --git a/Engine/Source/Programs/AutomationTool/Scripts/CopyBuildToStagingDirectory.Automation.cs b/Engine/Source/Programs/AutomationTool/Scripts/CopyBuildToStagingDirectory.Automation.cs
--- a/Engine/Source/Programs/AutomationTool/Scripts/CopyBuildToStagingDirectory.Automation.cs
+++ b/Engine/Source/Programs/AutomationTool/Scripts/CopyBuildToStagingDirectory.Automation.cs
@@ -490,6 +490,9 @@ public partial class Project : CommandUtils
                 SC.StageFiles(StagedFileTypeForMovies, CombinePaths(SC.ProjectRoot, "Content/Movies"), "*", true, new string[] { "*.uasset", "*.umap" }, CombinePaths(SC.RelativeProjectRootForStage, "Content/Movies"), true, !Params.UsePak(SC.StageTargetPlatform));
             }
 
+			// if we have oodle stuff 
+			SC.StageFiles( StagedFileType.UFS, CombinePaths( SC.ProjectRoot, "Content/Oodle" ), "*", true, new string[] { "*.uasset", "*.umap" }, CombinePaths( SC.RelativeProjectRootForStage, "Content/Oodle" ), true, !Params.UsePak( SC.StageTargetPlatform ) );
+			
             // eliminate the sand box
             SC.StageFiles(StagedFileType.UFS, CombinePaths(SC.ProjectRoot, "Saved", "Cooked", SC.CookPlatform), "*", true, new string[] { "*.json" }, "", true, !Params.UsePak(SC.StageTargetPlatform));
 
@@ -847,8 +850,15 @@ public partial class Project : CommandUtils
 	/// <param name="PakName"></param>
 	private static void CreatePak(ProjectParams Params, DeploymentContext SC, Dictionary<string, string> UnrealPakResponseFile, string PakName)
 	{
+        bool bShouldGeneratePatch = Params.IsGeneratingPatch && SC.StageTargetPlatform.GetPlatformPatchesWithDiffPak(Params, SC);
+
+        if (bShouldGeneratePatch && !Params.HasBasedOnReleaseVersion)
+        {
+            Log("Generating patch required a based on release version flag");
+        }
+
         string PostFix = "";
-		if (Params.IsGeneratingPatch && SC.StageTargetPlatform.GetPlatformPatchesWithDiffPak())
+		if (bShouldGeneratePatch)
         {
             PostFix += "_P";
         }
@@ -905,22 +915,11 @@ public partial class Project : CommandUtils
 		}
 
         string PatchSourceContentPath = null;
-		if (Params.HasBasedOnReleaseVersion && Params.IsGeneratingPatch && SC.StageTargetPlatform.GetPlatformPatchesWithDiffPak())
+		if (bShouldGeneratePatch)
         {
             // don't include the post fix in this filename because we are looking for the source pak path
             string PakFilename = PakName + "-" + SC.FinalCookPlatform + ".pak";
             PatchSourceContentPath = GetReleasePakFilePath(SC, Params, PakFilename);
-            
-
-            /*PatchSourceContentPath = Params.BasedOnReleaseVersion;
-            if (PatchSourceContentPath.EndsWith(PakName) == false)
-            {
-                string Temp = Path.GetDirectoryName(PatchSourceContentPath) + PakName;
-                if ( File.Exists(Temp))
-                {
-                    PatchSourceContentPath = Temp;
-                }
-            }*/
         }
 
 		if (!bCopiedExistingPak)
diff --git a/Engine/Source/Programs/CrashReporter/CrashReportClient/Private/CrashDescription.cpp b/Engine/Source/Programs/CrashReporter/CrashReportClient/Private/CrashDescription.cpp
--- a/Engine/Source/Programs/CrashReporter/CrashReportClient/Private/CrashDescription.cpp
+++ b/Engine/Source/Programs/CrashReporter/CrashReportClient/Private/CrashDescription.cpp
@@ -110,6 +110,7 @@ FPrimaryCrashProperties::FPrimaryCrashProperties()
 	, SourceContext( FGenericCrashContext::RuntimePropertiesTag, TEXT( "SourceContext" ), this )
 	, Modules( FGenericCrashContext::RuntimePropertiesTag, TEXT( "Modules" ), this )
 	, UserDescription( FGenericCrashContext::RuntimePropertiesTag, TEXT( "UserDescription" ), this )
+	, UserActivityHint( FGenericCrashContext::RuntimePropertiesTag, TEXT( "UserActivityHint" ), this )
 	, ErrorMessage( FGenericCrashContext::RuntimePropertiesTag, TEXT( "ErrorMessage" ), this )
 	, FullCrashDumpLocation( FGenericCrashContext::RuntimePropertiesTag, TEXT( "FullCrashDumpLocation" ), this )
 	, TimeOfCrash( FGenericCrashContext::RuntimePropertiesTag, TEXT( "TimeOfCrash" ), this )
@@ -209,6 +210,8 @@ void FPrimaryCrashProperties::SendAnalytics()
 	CrashAttributes.Add( FAnalyticsEventAttribute( TEXT( "EngineMode" ), EngineMode ) );
 	CrashAttributes.Add( FAnalyticsEventAttribute( TEXT( "AppDefaultLocale" ), AppDefaultLocale ) );
 
+	CrashAttributes.Add( FAnalyticsEventAttribute( TEXT( "UserActivityHint" ), UserActivityHint.AsString() ) );
+
 	Analytics.RecordEvent( TEXT( "CrashReportClient.ReportCrash" ), CrashAttributes );
 
 	// Shutdown analytics.
diff --git a/Engine/Source/Programs/CrashReporter/CrashReportClient/Private/CrashReportClientApp.cpp b/Engine/Source/Programs/CrashReporter/CrashReportClient/Private/CrashReportClientApp.cpp
--- a/Engine/Source/Programs/CrashReporter/CrashReportClient/Private/CrashReportClientApp.cpp
+++ b/Engine/Source/Programs/CrashReporter/CrashReportClient/Private/CrashReportClientApp.cpp
@@ -82,7 +82,7 @@ void ParseCommandLine(const TCHAR* CommandLine)
 			ReportDirectoryAbsolutePath = Tokens[0];
 		}
 
-		GameNameFromCmd = Params.FindRef("AppName=");
+		GameNameFromCmd = Params.FindRef("AppName");
 	}
 
 	if (ReportDirectoryAbsolutePath.IsEmpty())
@@ -163,7 +163,7 @@ void RunCrashReportClient(const TCHAR* CommandLine)
 	FPlatformErrorReport::Init();
 	auto ErrorReport = LoadErrorReport();
 	
-	if (ErrorReport.HasFilesToUpload())
+	if (ErrorReport.HasFilesToUpload() && FPrimaryCrashProperties::Get() != nullptr)
 	{
 		if (bUnattended)
 		{
diff --git a/Engine/Source/Programs/NetworkProfiler/NetworkProfiler/StreamParser.cs b/Engine/Source/Programs/NetworkProfiler/NetworkProfiler/StreamParser.cs
--- a/Engine/Source/Programs/NetworkProfiler/NetworkProfiler/StreamParser.cs
+++ b/Engine/Source/Programs/NetworkProfiler/NetworkProfiler/StreamParser.cs
@@ -88,6 +88,12 @@ namespace NetworkProfiler
 
 			// Parse stream till we reach the end, marked by special token.
 			bool bHasReachedEndOfStream = false;
+
+			List<TokenBase> TokenList = new List<TokenBase>();
+
+			float FrameStartTime = -1.0f;
+			float FrameEndTime = -1.0f;
+
 			while( bHasReachedEndOfStream == false )
 			{
 				if ( Count++ % 1000 == 0 )
@@ -134,10 +140,45 @@ namespace NetworkProfiler
 
 				if ( Token.TokenType == ETokenTypes.ConnectionChange )
 				{
+					// We need to setup CurrentConnectionIndex, since it's used in ReadNextToken
 					NetworkStream.CurrentConnectionIndex = ( Token as TokenConnectionChanged ).AddressIndex;
 					continue;
 				}
 
+				TokenList.Add( Token );
+
+				// Track frame start/end times manually so we can bail out when we hit the amount of time we want to load
+				if ( Token.TokenType == ETokenTypes.FrameMarker )
+				{
+					var TokenFrameMarker = ( TokenFrameMarker )Token;
+
+					if ( FrameStartTime < 0 )
+					{
+						FrameStartTime = TokenFrameMarker.RelativeTime;
+						FrameEndTime = TokenFrameMarker.RelativeTime;
+					}
+					else
+					{
+						FrameEndTime = TokenFrameMarker.RelativeTime;
+					}
+				}
+
+				if ( EarlyOutMinutes > 0 && ( ( FrameEndTime - FrameStartTime ) > 60 * EarlyOutMinutes ) )
+				{
+					break;
+				}
+			}
+
+			for ( int i = 0; i < TokenList.Count; i++ )
+			{
+				if ( i % 1000 == 0 )
+				{
+					float Percent = ( float )( i + 1 ) / ( float )( TokenList.Count );
+					InMainWindow.UpdateProgress( ( int )( Percent * 100 ) );
+				}
+
+				TokenBase Token = TokenList[i];
+
 				// Convert current tokens to frame if we reach a frame boundary or the end of the stream.
 				if( ((Token.TokenType == ETokenTypes.FrameMarker) || (Token.TokenType == ETokenTypes.EndOfStreamMarker))
 				// Nothing to do if we don't have any tokens, e.g. first frame.
@@ -220,11 +261,6 @@ namespace NetworkProfiler
 						CurrentFrameTokens.Add(Token);
 					}
 				}
-
-				if ( EarlyOutMinutes > 0 && ( ( AllFrames.EndTime - AllFrames.StartTime  ) > 60 * EarlyOutMinutes ) )
-				{
-					break;
-				}
 			}
 
 			InMainWindow.SetCurrentStreamSelection( NetworkStream, AllFrames, false );
diff --git a/Engine/Source/Programs/UnrealBuildTool/Configuration/BuildConfiguration.cs b/Engine/Source/Programs/UnrealBuildTool/Configuration/BuildConfiguration.cs
--- a/Engine/Source/Programs/UnrealBuildTool/Configuration/BuildConfiguration.cs
+++ b/Engine/Source/Programs/UnrealBuildTool/Configuration/BuildConfiguration.cs
@@ -537,6 +537,12 @@ namespace UnrealBuildTool
 		[XmlConfig]
 		public static bool bUseFastPDBLinking;
 
+		/// <summary>
+		/// Whether to request the linker create a map file as part of the build
+		/// </summary>
+		[XmlConfig]
+		public static bool bCreateMapFile;
+
 		/// <summary>
 		/// Sets the configuration back to defaults.
 		/// </summary>
@@ -701,6 +707,9 @@ namespace UnrealBuildTool
 			bAddFastPDBToProjects = false;
 			bUseFastPDBLinking = false;
 
+			// Don't create a map file by default
+			bCreateMapFile = false;
+
 			// The default for normal Mac users should be to use DistCode which installs as an Xcode plugin and provides dynamic host management
 			if (BuildHostPlatform.Current.Platform == UnrealTargetPlatform.Mac)
 			{
diff --git a/Engine/Source/Programs/UnrealBuildTool/System/UnrealBuildTool.cs b/Engine/Source/Programs/UnrealBuildTool/System/UnrealBuildTool.cs
--- a/Engine/Source/Programs/UnrealBuildTool/System/UnrealBuildTool.cs
+++ b/Engine/Source/Programs/UnrealBuildTool/System/UnrealBuildTool.cs
@@ -833,6 +833,10 @@ namespace UnrealBuildTool
 				{
 					BuildConfiguration.bUseFastPDBLinking = true;
 				}
+				else if(LowercaseArg == "-mapfile")
+				{
+					BuildConfiguration.bCreateMapFile = true;
+				}
 			}
 		}
 
diff --git a/Engine/Source/Programs/UnrealBuildTool/Windows/VCToolChain.cs b/Engine/Source/Programs/UnrealBuildTool/Windows/VCToolChain.cs
--- a/Engine/Source/Programs/UnrealBuildTool/Windows/VCToolChain.cs
+++ b/Engine/Source/Programs/UnrealBuildTool/Windows/VCToolChain.cs
@@ -5,6 +5,7 @@ using System.Collections.Generic;
 using System.Text.RegularExpressions;
 using System.Diagnostics;
 using System.IO;
+using System.Linq;
 using Microsoft.Win32;
 using System.Text;
 
@@ -1650,15 +1651,16 @@ namespace UnrealBuildTool
 					}
 
 					// Write the MAP file to the output directory.			
-#if false					
-					if (true)
+					if (BuildConfiguration.bCreateMapFile)
 					{
 						FileReference MAPFilePath = FileReference.Combine(LinkEnvironment.Config.OutputDirectory, Path.GetFileNameWithoutExtension(OutputFile.AbsolutePath) + ".map");
 						FileItem MAPFile = FileItem.GetItemByFileReference(MAPFilePath);
 						Arguments.AppendFormat(" /MAP:\"{0}\"", MAPFilePath);
-						LinkAction.ProducedItems.Add(MAPFile);
+						ProducedItems.Add(MAPFile);
+
+						// Export a list of object file paths, so we can locate the object files referenced by the map file
+						ExportObjectFilePaths(LinkEnvironment, Path.ChangeExtension(MAPFilePath.FullName, ".objpaths"));
 					}
-#endif
 				}
 
 				// Add the additional arguments specified by the environment.
@@ -1699,6 +1701,33 @@ namespace UnrealBuildTool
 			return OutputFile;
 		}
 
+		private void ExportObjectFilePaths(LinkEnvironment LinkEnvironment, string FileName)
+		{
+			// Write the list of object file directories
+			HashSet<DirectoryReference> ObjectFileDirectories = new HashSet<DirectoryReference>();
+			foreach(FileItem InputFile in LinkEnvironment.InputFiles)
+			{
+				ObjectFileDirectories.Add(InputFile.Reference.Directory);
+			}
+			foreach(FileItem InputLibrary in LinkEnvironment.InputLibraries)
+			{
+				ObjectFileDirectories.Add(InputLibrary.Reference.Directory);
+			}
+			foreach(string AdditionalLibrary in LinkEnvironment.Config.AdditionalLibraries.Where(x => Path.IsPathRooted(x)))
+			{
+				ObjectFileDirectories.Add(new FileReference(AdditionalLibrary).Directory);
+			}
+			foreach(string LibraryPath in LinkEnvironment.Config.LibraryPaths)
+			{
+				ObjectFileDirectories.Add(new DirectoryReference(LibraryPath));
+			}
+			foreach(string LibraryPath in (Environment.GetEnvironmentVariable("LIB") ?? "").Split(new char[]{ ';' }, StringSplitOptions.RemoveEmptyEntries))
+			{
+				ObjectFileDirectories.Add(new DirectoryReference(LibraryPath));
+			}
+			File.WriteAllLines(FileName, ObjectFileDirectories.Select(x => x.FullName).OrderBy(x => x).ToArray());
+		}
+
 		public override void CompileCSharpProject(CSharpEnvironment CompileEnvironment, FileReference ProjectFileName, FileReference DestinationFile)
 		{
 			// Initialize environment variables required for spawned tools.
diff --git a/Engine/Source/Programs/UnrealPak/Private/UnrealPak.cpp b/Engine/Source/Programs/UnrealPak/Private/UnrealPak.cpp
--- a/Engine/Source/Programs/UnrealPak/Private/UnrealPak.cpp
+++ b/Engine/Source/Programs/UnrealPak/Private/UnrealPak.cpp
@@ -1449,6 +1449,13 @@ INT32_MAIN_INT32_ARGC_TCHAR_ARGV()
 
 
 				Result = CreatePakFile(*PakFilename, FilesToAdd, CmdLineParameters) ? 0 : 1;
+
+				if (CmdLineParameters.GeneratePatch)
+				{
+					FString OutputPath = FPaths::GetPath(PakFilename) / FString(TEXT("TempFiles"));
+					// delete the temporary directory
+					IFileManager::Get().DeleteDirectory(*OutputPath, false, true);
+				}
 			}
 		}
 	}
diff --git a/Engine/Source/Runtime/AIModule/Classes/Tasks/AITask_MoveTo.h b/Engine/Source/Runtime/AIModule/Classes/Tasks/AITask_MoveTo.h
--- a/Engine/Source/Runtime/AIModule/Classes/Tasks/AITask_MoveTo.h
+++ b/Engine/Source/Runtime/AIModule/Classes/Tasks/AITask_MoveTo.h
@@ -59,7 +59,7 @@ public:
 	/**
 	 *	@param 
 	 */
-	UFUNCTION(BlueprintCallable, Category = "AI|Tasks", meta = (AdvancedDisplay = "TaskOwner,AcceptanceRadius,StopOnOverlap,AcceptPartialPath,bUsePathfinding", DefaultToSelf = "TaskOwner", BlueprintInternalUseOnly = "TRUE", DisplayName = "Move To Location or Actor"))
+	UFUNCTION(BlueprintCallable, Category = "AI|Tasks", meta = (AdvancedDisplay = "AcceptanceRadius,StopOnOverlap,AcceptPartialPath,bUsePathfinding", DefaultToSelf = "Controller", BlueprintInternalUseOnly = "TRUE", DisplayName = "Move To Location or Actor"))
 	static UAITask_MoveTo* AIMoveTo(AAIController* Controller, FVector GoalLocation, AActor* GoalActor = nullptr, float AcceptanceRadius = -1.f, EAIOptionFlag::Type StopOnOverlap = EAIOptionFlag::Default, EAIOptionFlag::Type AcceptPartialPath = EAIOptionFlag::Default, bool bUsePathfinding = true, bool bLockAILogic = true);
 
 	void SetUp(AAIController* Controller, FVector GoalLocation, AActor* GoalActor = nullptr, float AcceptanceRadius = -1.f, bool bUsePathfinding = true, EAIOptionFlag::Type StopOnOverlap = EAIOptionFlag::Default, EAIOptionFlag::Type AcceptPartialPath = EAIOptionFlag::Default);
diff --git a/Engine/Source/Runtime/AIModule/Private/Navigation/PathFollowingComponent.cpp b/Engine/Source/Runtime/AIModule/Private/Navigation/PathFollowingComponent.cpp
--- a/Engine/Source/Runtime/AIModule/Private/Navigation/PathFollowingComponent.cpp
+++ b/Engine/Source/Runtime/AIModule/Private/Navigation/PathFollowingComponent.cpp
@@ -458,6 +458,7 @@ void UPathFollowingComponent::OnPathUpdated()
 void UPathFollowingComponent::Initialize()
 {
 	UpdateCachedComponents();
+	LocationSamples.Reserve(BlockDetectionSampleCount);
 }
 
 void UPathFollowingComponent::Cleanup()
@@ -1211,6 +1212,8 @@ bool UPathFollowingComponent::IsBlocked() const
 
 	if (LocationSamples.Num() == BlockDetectionSampleCount && BlockDetectionSampleCount > 0)
 	{
+		const float BlockDetectionDistanceSq = FMath::Square(BlockDetectionDistance);
+
 		FVector Center = FVector::ZeroVector;
 		for (int32 SampleIndex = 0; SampleIndex < LocationSamples.Num(); SampleIndex++)
 		{
@@ -1222,8 +1225,8 @@ bool UPathFollowingComponent::IsBlocked() const
 
 		for (int32 SampleIndex = 0; SampleIndex < LocationSamples.Num(); SampleIndex++)
 		{
-			const float TestDistance = FVector::DistSquared(*LocationSamples[SampleIndex], Center);
-			if (TestDistance > BlockDetectionDistance)
+			const float TestDistanceSq = FVector::DistSquared(*LocationSamples[SampleIndex], Center);
+			if (TestDistanceSq > BlockDetectionDistanceSq)
 			{
 				bBlocked = false;
 				break;
@@ -1487,6 +1490,7 @@ void UPathFollowingComponent::DescribeSelfToVisLog(FVisualLogEntry* Snapshot) co
 		(Path->CastPath<FNavMeshPath>() != NULL) ? TEXT("navmesh") :
 		(Path->CastPath<FAbstractNavigationPath>() != NULL) ? TEXT("direct") :
 		TEXT("unknown"));
+	Category.Add(TEXT("Block detection"), bUseBlockDetection ? FString::Printf(TEXT("last sample time %.2f"), LastSampleTime) : TEXT("Disabled"));
 	
 	UObject* CustomNavLinkOb = CurrentCustomLinkOb.Get();
 	if (CustomNavLinkOb)
diff --git a/Engine/Source/Runtime/AIModule/Private/Tasks/AITask_MoveTo.cpp b/Engine/Source/Runtime/AIModule/Private/Tasks/AITask_MoveTo.cpp
--- a/Engine/Source/Runtime/AIModule/Private/Tasks/AITask_MoveTo.cpp
+++ b/Engine/Source/Runtime/AIModule/Private/Tasks/AITask_MoveTo.cpp
@@ -163,7 +163,7 @@ void UAITask_MoveTo::Pause()
 
 void UAITask_MoveTo::Resume()
 {
-	Activate();
+	PerformMove();
 	Super::Resume();
 }
 
diff --git a/Engine/Source/Runtime/Analytics/AnalyticsET/Private/AnalyticsET.cpp b/Engine/Source/Runtime/Analytics/AnalyticsET/Private/AnalyticsET.cpp
--- a/Engine/Source/Runtime/Analytics/AnalyticsET/Private/AnalyticsET.cpp
+++ b/Engine/Source/Runtime/Analytics/AnalyticsET/Private/AnalyticsET.cpp
@@ -13,6 +13,7 @@
 
 IMPLEMENT_MODULE( FAnalyticsET, AnalyticsET );
 DEFINE_LOG_CATEGORY(LogAnalytics);
+DEFINE_LOG_CATEGORY_STATIC(LogAnalyticsDumpEventPayload, Display, All);
 
 class FAnalyticsProviderET : 
 	public IAnalyticsProvider,
@@ -347,7 +348,7 @@ void FAnalyticsProviderET::FlushEvents()
 
 			// Recreate the URLPath for logging because we do not want to escape the parameters when logging.
 			// We cannot simply UrlEncode the entire Path after logging it because UrlEncode(Params) != UrlEncode(Param1) & UrlEncode(Param2) ...
-			UE_LOG(LogAnalytics, VeryVerbose, TEXT("[%s] AnalyticsET URL:datarouter/api/v1/public/data?SessionID=%s&AppID=%s&AppVersion=%s&UserID=%s&AppEnvironment=%s&UploadType=%s. Payload:%s"),
+			FString LogString = FString::Printf(TEXT("[%s] AnalyticsET URL:datarouter/api/v1/public/data?SessionID=%s&AppID=%s&AppVersion=%s&UserID=%s&AppEnvironment=%s&UploadType=%s. Payload:%s"),
 				*APIKey,
 				*SessionID,
 				*APIKey,
@@ -356,6 +357,11 @@ void FAnalyticsProviderET::FlushEvents()
 				*AppEnvironment,
 				*UploadType,
 				*Payload);
+			UE_LOG(LogAnalytics, VeryVerbose, TEXT("%s"), *LogString);
+
+			// Duplicate the same log message with a separate category. This is used as an "last chance" backup on the servers in the unlikely case 
+			// if the backend lost the events due to overload - then we can scrape the logs manually for them.
+			UE_LOG(LogAnalyticsDumpEventPayload, Log, TEXT("%s"), *LogString);
 
 			// Create/send Http request for an event
 			TSharedRef<IHttpRequest> HttpRequest = FHttpModule::Get().CreateRequest();
diff --git a/Engine/Source/Runtime/Analytics/QoSReporter/Private/QoSReporter.cpp b/Engine/Source/Runtime/Analytics/QoSReporter/Private/QoSReporter.cpp
--- a/Engine/Source/Runtime/Analytics/QoSReporter/Private/QoSReporter.cpp
+++ b/Engine/Source/Runtime/Analytics/QoSReporter/Private/QoSReporter.cpp
@@ -127,6 +127,7 @@ FString FQoSReporter::GetQoSReporterInstanceId()
 
 double FQoSReporter::ModuleInitializationTime = FPlatformTime::Seconds();
 bool FQoSReporter::bStartupEventReported = false;
+bool FQoSReporter::bCountHitches = false;
 
 void FQoSReporter::ReportStartupCompleteEvent()
 {
@@ -169,6 +170,12 @@ void FQoSReporter::SetBackendDeploymentName(const FString & InDeploymentName)
 	}
 }
 
+void FQoSReporter::EnableCountingHitches(bool bEnable)
+{
+	bCountHitches = bEnable;
+	UE_LOG(LogQoSReporter, Verbose, TEXT("Counting hitches in QoSReporter has been %s."), bCountHitches ? TEXT("enabled") : TEXT("disabled"));
+}
+
 double FQoSReporter::HeartbeatInterval = 300;
 double FQoSReporter::LastHeartbeatTimestamp = 0;
 extern ENGINE_API float GAverageFPS;
@@ -190,19 +197,24 @@ void FQoSReporter::Tick()
 	}
 
 	// detect too long pauses between ticks, unless configured to ignore them or running under debugger
-	if (!QOS_IGNORE_HITCHES && !FPlatformMisc::IsDebuggerPresent())
+	if (!QOS_IGNORE_HITCHES && bCountHitches && !FPlatformMisc::IsDebuggerPresent())
 	{
-		const double kTooLongPauseBetweenTicksInSeconds = 0.25f;	// 4 fps	- correct perfcounter name below if changing this value and make sure analytics is on the same page
 		const double DeltaBetweenTicks = CurrentTime - PreviousTickTime;
-		static int TimesHappened = 0;
-		if (DeltaBetweenTicks > kTooLongPauseBetweenTicksInSeconds)
+
+		if (DeltaBetweenTicks > 0.1)
 		{
-			++TimesHappened;
+#if USE_SERVER_PERF_COUNTERS
+			PerfCountersIncrement(TEXT("HitchesAbove100msec"), 0, IPerfCounters::Flags::Transient);
 
-			PerfCountersIncrement(TEXT("HitchesAbove250msec"));
+			// count 250 ms
+			if (DeltaBetweenTicks > 0.25)
+			{
+				PerfCountersIncrement(TEXT("HitchesAbove250msec"), 0, IPerfCounters::Flags::Transient);
+			}
+#endif // USE_SERVER_PERF_COUNTERS
 
-			UE_LOG(LogQoSReporter, Warning, TEXT("QoS reporter could not tick for %f sec (longer than threshold of %f sec) - happened %d time(s), average FPS is %f. Sending heartbeats might have been affected"),
-				DeltaBetweenTicks, kTooLongPauseBetweenTicksInSeconds, TimesHappened, GAverageFPS);
+			UE_LOG(LogQoSReporter, Warning, TEXT("QoS reporter could not tick for %f sec, average FPS is %f."),
+				DeltaBetweenTicks, GAverageFPS);
 		}
 	}
 
diff --git a/Engine/Source/Runtime/Analytics/QoSReporter/Private/QoSReporterModule.cpp b/Engine/Source/Runtime/Analytics/QoSReporter/Private/QoSReporterModule.cpp
--- a/Engine/Source/Runtime/Analytics/QoSReporter/Private/QoSReporterModule.cpp
+++ b/Engine/Source/Runtime/Analytics/QoSReporter/Private/QoSReporterModule.cpp
@@ -14,7 +14,7 @@ DEFINE_LOG_CATEGORY(LogQoSReporter);
 IMPLEMENT_MODULE(FQoSReporterModule, QoSReporter);
 
 // helps to version QoS events (date*10 to allow for 10 revisions per day)
-#define QOS_EVENTS_REVISION					201512170
+#define QOS_EVENTS_REVISION					201601200
 
 FString FQoSReporterModule::Config::GetDefaultAppVersion()
 { 
diff --git a/Engine/Source/Runtime/Analytics/QoSReporter/Public/QoSReporter.h b/Engine/Source/Runtime/Analytics/QoSReporter/Public/QoSReporter.h
--- a/Engine/Source/Runtime/Analytics/QoSReporter/Public/QoSReporter.h
+++ b/Engine/Source/Runtime/Analytics/QoSReporter/Public/QoSReporter.h
@@ -169,6 +169,13 @@ public:
 	 */
 	static void QOSREPORTER_API SetBackendDeploymentName(const FString & InDeploymentName);
 
+	/**
+	 * Whether or not hitches should be counted
+	 *
+	 * @param bEnable true if yes, false if no (default)
+	 */
+	static void QOSREPORTER_API EnableCountingHitches(bool bEnable);
+
 private:
 
 	/**
@@ -206,4 +213,7 @@ private:
 
 	/** Whether startup event was reported. */
 	static bool bStartupEventReported;
+
+	/** Whether or not counting hitches has been suppressed. */
+	static bool bCountHitches;
 };
diff --git a/Engine/Source/Runtime/AnimGraphRuntime/Private/AnimNodes/AnimNode_SequenceEvaluator.cpp b/Engine/Source/Runtime/AnimGraphRuntime/Private/AnimNodes/AnimNode_SequenceEvaluator.cpp
--- a/Engine/Source/Runtime/AnimGraphRuntime/Private/AnimNodes/AnimNode_SequenceEvaluator.cpp
+++ b/Engine/Source/Runtime/AnimGraphRuntime/Private/AnimNodes/AnimNode_SequenceEvaluator.cpp
@@ -25,7 +25,7 @@ void FAnimNode_SequenceEvaluator::UpdateAssetPlayer(const FAnimationUpdateContex
 		// Clamp input to a valid position on this sequence's time line.
 		ExplicitTime = FMath::Clamp(ExplicitTime, 0.f, Sequence->SequenceLength);
 
-		if ((GroupIndex != INDEX_NONE) && (Context.AnimInstanceProxy->IsSkeletonCompatible(Sequence->GetSkeleton())))
+		if ((!bTeleportToExplicitTime || (GroupIndex != INDEX_NONE)) && (Context.AnimInstanceProxy->IsSkeletonCompatible(Sequence->GetSkeleton())))
 		{
 			InternalTimeAccumulator = FMath::Clamp(InternalTimeAccumulator, 0.f, Sequence->SequenceLength);
 			float TimeJump = ExplicitTime - InternalTimeAccumulator;
diff --git a/Engine/Source/Runtime/AnimGraphRuntime/Public/AnimNodes/AnimNode_SequenceEvaluator.h b/Engine/Source/Runtime/AnimGraphRuntime/Public/AnimNodes/AnimNode_SequenceEvaluator.h
--- a/Engine/Source/Runtime/AnimGraphRuntime/Public/AnimNodes/AnimNode_SequenceEvaluator.h
+++ b/Engine/Source/Runtime/AnimGraphRuntime/Public/AnimNodes/AnimNode_SequenceEvaluator.h
@@ -24,9 +24,16 @@ public:
 	UPROPERTY(EditAnywhere, BlueprintReadWrite, Category = Settings, meta = (PinHiddenByDefault))
 	bool bShouldLoopWhenInSyncGroup;
 
+	/** If true, teleport to explicit time, does NOT advance time (does not trigger notifies, does not extract Root Motion, etc.)
+	If false, will advance time (will trigger notifies, extract root motion if applicable, etc.)
+	Note: using a sync group forces advancing time regardless of what this option is set to. */
+	UPROPERTY(EditAnywhere, BlueprintReadWrite, Category = Settings, meta = (PinHiddenByDefault))
+	bool bTeleportToExplicitTime;
+
 public:	
 	FAnimNode_SequenceEvaluator()
 		: ExplicitTime(0.0f)
+		, bTeleportToExplicitTime(true)
 	{
 	}
 
diff --git a/Engine/Source/Runtime/Core/Private/GenericPlatform/GenericApplication.cpp b/Engine/Source/Runtime/Core/Private/GenericPlatform/GenericApplication.cpp
--- a/Engine/Source/Runtime/Core/Private/GenericPlatform/GenericApplication.cpp
+++ b/Engine/Source/Runtime/Core/Private/GenericPlatform/GenericApplication.cpp
@@ -17,6 +17,8 @@ const FGamepadKeyNames::Type FGamepadKeyNames::RightTriggerAnalog("Gamepad_Right
 const FGamepadKeyNames::Type FGamepadKeyNames::LeftThumb("Gamepad_LeftThumbstick");
 const FGamepadKeyNames::Type FGamepadKeyNames::RightThumb("Gamepad_RightThumbstick");
 const FGamepadKeyNames::Type FGamepadKeyNames::SpecialLeft("Gamepad_Special_Left");
+const FGamepadKeyNames::Type FGamepadKeyNames::SpecialLeft_X("Gamepad_Special_Left_X");
+const FGamepadKeyNames::Type FGamepadKeyNames::SpecialLeft_Y("Gamepad_Special_Left_Y");
 const FGamepadKeyNames::Type FGamepadKeyNames::SpecialRight("Gamepad_Special_Right");
 const FGamepadKeyNames::Type FGamepadKeyNames::FaceButtonBottom("Gamepad_FaceButton_Bottom");
 const FGamepadKeyNames::Type FGamepadKeyNames::FaceButtonRight("Gamepad_FaceButton_Right");
diff --git a/Engine/Source/Runtime/Core/Private/GenericPlatform/GenericPlatformCrashContext.cpp b/Engine/Source/Runtime/Core/Private/GenericPlatform/GenericPlatformCrashContext.cpp
--- a/Engine/Source/Runtime/Core/Private/GenericPlatform/GenericPlatformCrashContext.cpp
+++ b/Engine/Source/Runtime/Core/Private/GenericPlatform/GenericPlatformCrashContext.cpp
@@ -48,6 +48,7 @@ namespace NCachedCrashContextProperties
 	static int32 CrashDumpMode;
 	static int32 SecondsSinceStart;
 	static FString CrashGUID;
+	static FString UserActivityHint;
 }
 
 void FGenericCrashContext::Initialize()
@@ -101,6 +102,11 @@ void FGenericCrashContext::Initialize()
 		return true;
 	} ), PollingInterval );
 
+	FCoreDelegates::UserActivityStringChanged.AddLambda([](const FString& InUserActivity)
+	{
+		NCachedCrashContextProperties::UserActivityHint = InUserActivity;
+	});
+
 	bIsInitialized = true;
 }
 
@@ -153,6 +159,7 @@ void FGenericCrashContext::SerializeContentToBuffer()
 	AddCrashProperty( TEXT( "CallStack" ), TEXT( "" ) );
 	AddCrashProperty( TEXT( "SourceContext" ), TEXT( "" ) );
 	AddCrashProperty( TEXT( "UserDescription" ), TEXT( "" ) );
+	AddCrashProperty( TEXT( "UserActivityHint" ), *NCachedCrashContextProperties::UserActivityHint );
 	AddCrashProperty( TEXT( "ErrorMessage" ), (const TCHAR*)GErrorMessage ); // GErrorMessage may be broken.
 	AddCrashProperty( TEXT( "CrashDumpMode" ), NCachedCrashContextProperties::CrashDumpMode );
 
diff --git a/Engine/Source/Runtime/Core/Private/GenericPlatform/GenericPlatformMisc.cpp b/Engine/Source/Runtime/Core/Private/GenericPlatform/GenericPlatformMisc.cpp
--- a/Engine/Source/Runtime/Core/Private/GenericPlatform/GenericPlatformMisc.cpp
+++ b/Engine/Source/Runtime/Core/Private/GenericPlatform/GenericPlatformMisc.cpp
@@ -413,6 +413,14 @@ void FGenericPlatformMisc::RequestExit( bool Force )
 	}
 }
 
+void FGenericPlatformMisc::RequestExitWithStatus(bool Force, uint8 ReturnCode)
+{
+	// Generic implementation will ignore the return code - this may be important, so warn.
+	UE_LOG(LogGenericPlatformMisc, Warning, TEXT("FPlatformMisc::RequestExitWithStatus(%i, %d) - return code will be ignored by the generic implementation."), Force, ReturnCode);
+
+	return FPlatformMisc::RequestExit(Force);
+}
+
 const TCHAR* FGenericPlatformMisc::GetSystemErrorMessage(TCHAR* OutBuffer, int32 BufferCount, int32 Error)
 {
 	const TCHAR* Message = TEXT("No system errors available on this platform.");
diff --git a/Engine/Source/Runtime/Core/Private/HAL/Allocators/CachedOSPageAllocator.cpp b/Engine/Source/Runtime/Core/Private/HAL/Allocators/CachedOSPageAllocator.cpp
--- a/Engine/Source/Runtime/Core/Private/HAL/Allocators/CachedOSPageAllocator.cpp
+++ b/Engine/Source/Runtime/Core/Private/HAL/Allocators/CachedOSPageAllocator.cpp
@@ -38,6 +38,7 @@ void* FCachedOSPageAllocator::AllocateImpl(SIZE_T Size, FFreePageBlock* First, F
 		if (Found)
 		{
 			void* Result = Found->Ptr;
+			UE_CLOG(!Result, LogMemory, Fatal, TEXT("OS memory allocation cache has been corrupted!"));
 			CachedTotal -= Found->ByteSize;
 			if (Found + 1 != Last)
 			{
diff --git a/Engine/Source/Runtime/Core/Private/HAL/ConsoleManager.cpp b/Engine/Source/Runtime/Core/Private/HAL/ConsoleManager.cpp
--- a/Engine/Source/Runtime/Core/Private/HAL/ConsoleManager.cpp
+++ b/Engine/Source/Runtime/Core/Private/HAL/ConsoleManager.cpp
@@ -807,6 +807,33 @@ IConsoleObject* FConsoleManager::FindConsoleObject(const TCHAR* Name) const
 {
 	IConsoleObject* CVar = FindConsoleObjectUnfiltered(Name);
 
+#if TRACK_CONSOLE_FIND_COUNT
+	{
+		const bool bEarlyAppPhase = GFrameCounter < 1000;
+		if(CVar)
+		{
+			++CVar->FindCallCount;
+
+			// we test for equal to avoid log spam
+			if(bEarlyAppPhase && CVar->FindCallCount == 500)
+			{
+				UE_LOG(LogConsoleManager, Warning, TEXT("Performance warning: Console object named '%s' shows many (%d) FindConsoleObject() calls (consider caching e.g. using static)"), Name, CVar->FindCallCount);
+			}
+		}
+		else
+		{
+			static uint32 NullFindCallCount = 0;
+		
+			++NullFindCallCount;
+
+			if(bEarlyAppPhase && NullFindCallCount == 500)
+			{
+				UE_LOG(LogConsoleManager, Warning, TEXT( "Performance warning: Many (%d) failed FindConsoleObject() e.g. '%s' (consider caching, is the name referencing an existing object)"), NullFindCallCount, Name);
+			}
+		}
+	}
+#endif
+
 	if(CVar && CVar->TestFlags(ECVF_CreatedFromIni))
 	{
 		return 0;
diff --git a/Engine/Source/Runtime/Core/Private/HAL/FileManagerGeneric.cpp b/Engine/Source/Runtime/Core/Private/HAL/FileManagerGeneric.cpp
--- a/Engine/Source/Runtime/Core/Private/HAL/FileManagerGeneric.cpp
+++ b/Engine/Source/Runtime/Core/Private/HAL/FileManagerGeneric.cpp
@@ -661,10 +661,39 @@ bool FArchiveFileReaderGeneric::InternalPrecache( int64 PrecacheOffset, int64 Pr
 		BufferCount = FMath::Max( BufferCount, 0LL ); // clamp to 0
 		int64 Count = 0;
 
-		// Read data from device via Win32 ReadFile API.
 		{
-			UE_CLOG( BufferCount > ARRAY_COUNT( Buffer ) || BufferCount <= 0, LogFileManager, Fatal, TEXT("Invalid BufferCount=%lld while reading %s. File is most likely corrupted, try deleting it if possible. Pos=%lld, Size=%lld, PrecacheSize=%lld, PrecacheOffset=%lld"),
-				BufferCount, *Filename, Pos, Size, PrecacheSize, PrecacheOffset );
+			#if PLATFORM_DESKTOP
+				// Show a message box indicating, possible, corrupt data (desktop platforms only)
+				if (BufferCount > ARRAY_COUNT( Buffer ) || BufferCount <= 0)
+				{
+					FText ErrorMessage, ErrorCaption;
+					GConfig->GetText(TEXT("/Script/Engine.Engine"),
+									 TEXT("SerializationOutOfBoundsErrorMessage"),
+									 ErrorMessage,
+									 GEngineIni);
+					GConfig->GetText(TEXT("/Script/Engine.Engine"),
+						TEXT("SerializationOutOfBoundsErrorMessageCaption"),
+						ErrorCaption,
+						GEngineIni);
+
+					UE_LOG( LogFileManager, Error, TEXT("Invalid BufferCount=%lld while reading %s. File is most likely corrupted. Please verify your installation. Pos=%lld, Size=%lld, PrecacheSize=%lld, PrecacheOffset=%lld"),
+						BufferCount, *Filename, Pos, Size, PrecacheSize, PrecacheOffset );
+
+					if (GLog)
+					{
+						GLog->Flush();
+					}
+
+					FPlatformMisc::MessageBoxExt(EAppMsgType::Ok, *ErrorMessage.ToString(), *ErrorCaption.ToString());
+
+					check(false);
+				}
+			#else
+				{
+					UE_CLOG( BufferCount > ARRAY_COUNT( Buffer ) || BufferCount <= 0, LogFileManager, Fatal, TEXT("Invalid BufferCount=%lld while reading %s. File is most likely corrupted. Please verify your installation. Pos=%lld, Size=%lld, PrecacheSize=%lld, PrecacheOffset=%lld"),
+						BufferCount, *Filename, Pos, Size, PrecacheSize, PrecacheOffset );
+				}
+			#endif
 
 			ReadLowLevel( Buffer, BufferCount, Count );
 		}
diff --git a/Engine/Source/Runtime/Core/Private/HAL/MallocBinned.cpp b/Engine/Source/Runtime/Core/Private/HAL/MallocBinned.cpp
--- a/Engine/Source/Runtime/Core/Private/HAL/MallocBinned.cpp
+++ b/Engine/Source/Runtime/Core/Private/HAL/MallocBinned.cpp
@@ -591,6 +591,7 @@ struct FMallocBinned::Private
 				if (Allocator.FreedPageBlocks[i].ByteSize == NewSize)
 				{
 					void* Ret = Allocator.FreedPageBlocks[i].Ptr;
+					UE_CLOG(!Ret, LogMemory, Fatal, TEXT("OS memory allocation cache has been corrupted!"));
 					OutActualSize = Allocator.FreedPageBlocks[i].ByteSize;
 					Allocator.CachedTotal -= Allocator.FreedPageBlocks[i].ByteSize;
 					if (i < Allocator.FreedPageBlocksNum - 1)
@@ -607,6 +608,7 @@ struct FMallocBinned::Private
 				if (Allocator.FreedPageBlocks[i].ByteSize >= NewSize && Allocator.FreedPageBlocks[i].ByteSize * 3 <= NewSize * 4)
 				{
 					void* Ret = Allocator.FreedPageBlocks[i].Ptr;
+					UE_CLOG(!Ret, LogMemory, Fatal, TEXT("OS memory allocation cache has been corrupted!"));
 					OutActualSize = Allocator.FreedPageBlocks[i].ByteSize;
 					Allocator.CachedTotal -= Allocator.FreedPageBlocks[i].ByteSize;
 					if (i < Allocator.FreedPageBlocksNum - 1)
diff --git a/Engine/Source/Runtime/Core/Private/Internationalization/Text.cpp b/Engine/Source/Runtime/Core/Private/Internationalization/Text.cpp
--- a/Engine/Source/Runtime/Core/Private/Internationalization/Text.cpp
+++ b/Engine/Source/Runtime/Core/Private/Internationalization/Text.cpp
@@ -1057,11 +1057,16 @@ CORE_API FArchive& operator<<(FArchive& Ar, FText& Value)
 		Value.TextData->PersistText(); // We always need to do this when saving so that we can save the history correctly
 		if(Ar.IsPersistent())
 		{
-			Value.Flags &= ~(ETextFlag::ConvertedProperty); // Remove conversion flag before saving.
+			Value.Flags &= ~(ETextFlag::ConvertedProperty | ETextFlag::InitializedFromString); // Remove conversion flag before saving.
 		}
 	}
 	Ar << Value.Flags;
 
+	if (Ar.IsLoading() && Ar.ArIsPersistent)
+	{
+		Value.Flags &= ~(ETextFlag::ConvertedProperty | ETextFlag::InitializedFromString); // Remove conversion flag before saving.
+	}
+
 	if (Ar.IsSaving())
 	{
 		Value.Flags = OriginalFlags;
diff --git a/Engine/Source/Runtime/Core/Private/Linux/LinuxPlatformCrashContext.cpp b/Engine/Source/Runtime/Core/Private/Linux/LinuxPlatformCrashContext.cpp
--- a/Engine/Source/Runtime/Core/Private/Linux/LinuxPlatformCrashContext.cpp
+++ b/Engine/Source/Runtime/Core/Private/Linux/LinuxPlatformCrashContext.cpp
@@ -77,7 +77,7 @@ void GracefulTerminationHandler(int32 Signal, siginfo_t* Info, void* Context)
 	// do not flush logs at this point; this can result in a deadlock if the signal was received while we were holding lock in the malloc (flushing allocates memory)
 	if( !GIsRequestingExit )
 	{
-		GIsRequestingExit = 1;
+		FPlatformMisc::RequestExitWithStatus(false, 128 + Signal);	// Keeping the established shell practice of returning 128 + signal for terminations by signal. Allows to distinguish SIGINT/SIGTERM/SIGHUP.
 	}
 	else
 	{
@@ -395,7 +395,16 @@ void DLLEXPORT GenerateCrashInfoAndLaunchReporter(const FLinuxCrashContext & Con
 
 #endif
 
-	FPlatformMisc::RequestExit(true);
+	// remove the handler for this signal and re-raise it (which should generate the proper core dump)
+	UE_LOG(LogLinux, Log, TEXT("Engine crash handling finished; re-raising signal %d for the default handler. Good bye."), Context.Signal);
+
+	struct sigaction ResetToDefaultAction;
+	FMemory::Memzero(ResetToDefaultAction);
+	ResetToDefaultAction.sa_handler = SIG_DFL;
+	sigemptyset(&ResetToDefaultAction.sa_mask);
+	sigaction(Context.Signal, &ResetToDefaultAction, nullptr);
+
+	raise(Context.Signal);
 }
 
 /**
diff --git a/Engine/Source/Runtime/Core/Private/Linux/LinuxPlatformMisc.cpp b/Engine/Source/Runtime/Core/Private/Linux/LinuxPlatformMisc.cpp
--- a/Engine/Source/Runtime/Core/Private/Linux/LinuxPlatformMisc.cpp
+++ b/Engine/Source/Runtime/Core/Private/Linux/LinuxPlatformMisc.cpp
@@ -368,6 +368,29 @@ uint32 FLinuxPlatformMisc::GetKeyMap( uint32* KeyCodes, FString* KeyNames, uint3
 	return NumMappings;
 }
 
+uint8 GOverriddenReturnCode = 0;
+bool GHasOverriddenReturnCode = false;
+
+void FLinuxPlatformMisc::RequestExitWithStatus(bool Force, uint8 ReturnCode)
+{
+	UE_LOG(LogLinux, Log, TEXT("FLinuxPlatformMisc::RequestExit(bForce=%s, ReturnCode=%d)"), Force ? TEXT("true") : TEXT("false"), ReturnCode);
+
+	GOverriddenReturnCode = ReturnCode;
+	GHasOverriddenReturnCode = true;
+
+	return FPlatformMisc::RequestExit(Force);
+}
+
+bool FLinuxPlatformMisc::HasOverriddenReturnCode(uint8 * OverriddenReturnCodeToUsePtr)
+{
+	if (GHasOverriddenReturnCode && OverriddenReturnCodeToUsePtr != nullptr)
+	{
+		*OverriddenReturnCodeToUsePtr = GOverriddenReturnCode;
+	}
+
+	return GHasOverriddenReturnCode;
+}
+
 const TCHAR* FLinuxPlatformMisc::GetSystemErrorMessage(TCHAR* OutBuffer, int32 BufferCount, int32 Error)
 {
 	check(OutBuffer && BufferCount);
diff --git a/Engine/Source/Runtime/Core/Private/Linux/LinuxPlatformOutputDevices.cpp b/Engine/Source/Runtime/Core/Private/Linux/LinuxPlatformOutputDevices.cpp
--- a/Engine/Source/Runtime/Core/Private/Linux/LinuxPlatformOutputDevices.cpp
+++ b/Engine/Source/Runtime/Core/Private/Linux/LinuxPlatformOutputDevices.cpp
@@ -42,6 +42,12 @@ void								FLinuxOutputDevices::SetupOutputDevices()
 	// debug and event logging is not really supported on Linux. 
 }
 
+FString								FLinuxOutputDevices::GetAbsoluteLogFilename()
+{
+	// FIXME: this function should not exist once FGenericPlatformOutputDevices::GetAbsoluteLogFilename() returns absolute filename (see UE-25650)
+	return FPaths::ConvertRelativePathToFull(FGenericPlatformOutputDevices::GetAbsoluteLogFilename());
+}
+
 class FOutputDeviceError*			FLinuxOutputDevices::GetError()
 {
 	static FOutputDeviceLinuxError Singleton;
diff --git a/Engine/Source/Runtime/Core/Private/Misc/CallbackDevice.cpp b/Engine/Source/Runtime/Core/Private/Misc/CallbackDevice.cpp
--- a/Engine/Source/Runtime/Core/Private/Misc/CallbackDevice.cpp
+++ b/Engine/Source/Runtime/Core/Private/Misc/CallbackDevice.cpp
@@ -70,4 +70,8 @@ FCoreDelegates::FPlatformChangedLaptopMode FCoreDelegates::PlatformChangedLaptop
 
 FCoreDelegates::FLoadStringAssetReferenceInCook FCoreDelegates::LoadStringAssetReferenceInCook;
 
-FCoreDelegates::FVRHeadsetRecenter FCoreDelegates::VRHeadsetRecenter;
\ No newline at end of file
+FCoreDelegates::FVRHeadsetRecenter FCoreDelegates::VRHeadsetRecenter;
+
+FCoreDelegates::FOnUserActivityStringChanged FCoreDelegates::UserActivityStringChanged;
+
+FCoreDelegates::FOnAsyncLoadingFlush FCoreDelegates::OnAsyncLoadingFlush;
\ No newline at end of file
diff --git a/Engine/Source/Runtime/Core/Private/Windows/WindowsApplication.cpp b/Engine/Source/Runtime/Core/Private/Windows/WindowsApplication.cpp
--- a/Engine/Source/Runtime/Core/Private/Windows/WindowsApplication.cpp
+++ b/Engine/Source/Runtime/Core/Private/Windows/WindowsApplication.cpp
@@ -23,6 +23,7 @@
 #include <SetupApi.h>
 #include <devguid.h>
 #include <dwmapi.h>
+#include <cfgmgr32.h>
 
 
 // This might not be defined by Windows when maintaining backwards-compatibility to pre-Vista builds
@@ -487,19 +488,26 @@ inline bool GetSizeForDevID(const FString& TargetDevID, int32& Width, int32& Hei
 
 		if (SetupDiEnumDeviceInfo(DevInfo, MonitorIndex, &DevInfoData) == TRUE)
 		{
-			HKEY hDevRegKey = SetupDiOpenDevRegKey(DevInfo, &DevInfoData, DICS_FLAG_GLOBAL, 0, DIREG_DEV, KEY_READ);
-
-			if(!hDevRegKey || (hDevRegKey == INVALID_HANDLE_VALUE))
+			TCHAR Buffer[MAX_DEVICE_ID_LEN];
+			if (CM_Get_Device_ID(DevInfoData.DevInst, Buffer, MAX_DEVICE_ID_LEN, 0) == CR_SUCCESS)
 			{
-				continue;
-			}
-
-			bRes = GetMonitorSizeFromEDID(hDevRegKey, Width, Height);
+				FString DevID(Buffer);
+				DevID = DevID.Mid(8, DevID.Find(TEXT("\\"), ESearchCase::CaseSensitive, ESearchDir::FromStart, 9) - 8);
+				if (DevID == TargetDevID)
+				{
+					HKEY hDevRegKey = SetupDiOpenDevRegKey(DevInfo, &DevInfoData, DICS_FLAG_GLOBAL, 0, DIREG_DEV, KEY_READ);
 
-			RegCloseKey(hDevRegKey);
+					if (hDevRegKey && hDevRegKey != INVALID_HANDLE_VALUE)
+					{
+						bRes = GetMonitorSizeFromEDID(hDevRegKey, Width, Height);
+						RegCloseKey(hDevRegKey);
+						break;
+					}
+				}
+			}
 		}
 	}
-	
+
 	if (SetupDiDestroyDeviceInfoList(DevInfo) == FALSE)
 	{
 		bRes = false;
diff --git a/Engine/Source/Runtime/Core/Private/Windows/WindowsPlatformCrashContext.cpp b/Engine/Source/Runtime/Core/Private/Windows/WindowsPlatformCrashContext.cpp
--- a/Engine/Source/Runtime/Core/Private/Windows/WindowsPlatformCrashContext.cpp
+++ b/Engine/Source/Runtime/Core/Private/Windows/WindowsPlatformCrashContext.cpp
@@ -337,10 +337,18 @@ int32 ReportCrashUsingCrashReportClient(EXCEPTION_POINTERS* ExceptionInfo, const
 
 		// Suppress the user input dialog if we're running in unattended mode
 		bool bNoDialog = FApp::IsUnattended() || ReportUI == EErrorReportUI::ReportInUnattendedMode || IsRunningDedicatedServer();
-		if( bNoDialog )
+		// Pass nullrhi to CRC when the engine is in this mode to stop the CRC attempting to initialize RHI when the capability isn't available
+		bool bNullRHI = FApp::ShouldUseNullRHI();
+
+		if (bNoDialog || bNullRHI)
 		{
 			CrashReportClientArguments += TEXT( " -Unattended" );
 		}
+		
+		if (bNullRHI)
+		{
+			CrashReportClientArguments += TEXT(" -nullrhi");
+		}
 
 		CrashReportClientArguments += FString( TEXT( " -AppName=" ) ) + ReportInformation.wzApplicationName;
 
