Fix compile errors for UE4Editor when Oodle SDK is not available.

diff --git a/Engine/Plugins/Runtime/PacketHandlers/CompressionComponents/Oodle/Source/OodleHandlerComponent/Private/OodleCustomCompressor.cpp b/Engine/Plugins/Runtime/PacketHandlers/CompressionComponents/Oodle/Source/OodleHandlerComponent/Private/OodleCustomCompressor.cpp
--- a/Engine/Plugins/Runtime/PacketHandlers/CompressionComponents/Oodle/Source/OodleHandlerComponent/Private/OodleCustomCompressor.cpp
+++ b/Engine/Plugins/Runtime/PacketHandlers/CompressionComponents/Oodle/Source/OodleHandlerComponent/Private/OodleCustomCompressor.cpp
@@ -1,6 +1,11 @@
 #include "CoreTypes.h"
 #include "Misc/Compression.h"
+
+#if HAS_OODLE_SDK
 #include "oodle2.h"
+#endif
+
+#if HAS_OODLE_SDK
 
 struct FOodleCustomCompressor : ICustomCompressor
 {
@@ -48,3 +53,5 @@ ICustomCompressor* CreateOodleCustomCompressor()
 {
 	return new FOodleCustomCompressor(OodleLZ_Compressor_Mermaid, OodleLZ_CompressionLevel_Optimal2, 256);
 }
+
+#endif
diff --git a/Engine/Plugins/Runtime/PacketHandlers/CompressionComponents/Oodle/Source/OodleHandlerComponent/Private/OodleUtils.cpp b/Engine/Plugins/Runtime/PacketHandlers/CompressionComponents/Oodle/Source/OodleHandlerComponent/Private/OodleUtils.cpp
--- a/Engine/Plugins/Runtime/PacketHandlers/CompressionComponents/Oodle/Source/OodleHandlerComponent/Private/OodleUtils.cpp
+++ b/Engine/Plugins/Runtime/PacketHandlers/CompressionComponents/Oodle/Source/OodleHandlerComponent/Private/OodleUtils.cpp
@@ -1,7 +1,9 @@
 #include "OodleUtils.h"
 #include "Serialization/MemoryReader.h"
 #include "Serialization/MemoryWriter.h"
+#if HAS_OODLE_SDK
 #include "oodle2.h"
+#endif
 
 #if HAS_OODLE_SDK
 
diff --git a/Engine/Plugins/Runtime/PacketHandlers/CompressionComponents/Oodle/UnrealPakPlugin/UnrealPakPlugin.cpp b/Engine/Plugins/Runtime/PacketHandlers/CompressionComponents/Oodle/UnrealPakPlugin/UnrealPakPlugin.cpp
--- a/Engine/Plugins/Runtime/PacketHandlers/CompressionComponents/Oodle/UnrealPakPlugin/UnrealPakPlugin.cpp
+++ b/Engine/Plugins/Runtime/PacketHandlers/CompressionComponents/Oodle/UnrealPakPlugin/UnrealPakPlugin.cpp
@@ -16,13 +16,14 @@
 	#else
 		#error Unknown platform
 	#endif
-#elif WIN64
+#elif _WINDOWS
 	#define PLATFORM_WINDOWS 1
 #else
 	#error Unknown platform
 #endif
 
-#define CORE_API 
+#define CORE_API
+#define HAS_OODLE_SDK 1
 
 #include <stdlib.h>
 #include "../Source/OodleHandlerComponent/Private/OodleCustomCompressor.cpp"
