Merging //UE4/Dev-Main to Dev-Anim (//UE4/Dev-Anim) @ CL 4107923

diff --git a/Engine/Plugins/Runtime/PacketHandlers/CompressionComponents/Oodle/Source/OodleHandlerComponent/Private/OodleCustomCompressor.cpp b/Engine/Plugins/Runtime/PacketHandlers/CompressionComponents/Oodle/Source/OodleHandlerComponent/Private/OodleCustomCompressor.cpp
--- a/Engine/Plugins/Runtime/PacketHandlers/CompressionComponents/Oodle/Source/OodleHandlerComponent/Private/OodleCustomCompressor.cpp
+++ b/Engine/Plugins/Runtime/PacketHandlers/CompressionComponents/Oodle/Source/OodleHandlerComponent/Private/OodleCustomCompressor.cpp
@@ -1,6 +1,11 @@
 #include "CoreTypes.h"
 #include "Misc/Compression.h"
+
+#if HAS_OODLE_SDK
 #include "oodle2.h"
+#endif
+
+#if HAS_OODLE_SDK
 
 struct FOodleCustomCompressor : ICustomCompressor
 {
@@ -48,3 +53,5 @@ ICustomCompressor* CreateOodleCustomCompressor()
 {
 	return new FOodleCustomCompressor(OodleLZ_Compressor_Mermaid, OodleLZ_CompressionLevel_Optimal2, 256);
 }
+
+#endif
diff --git a/Engine/Plugins/Runtime/PacketHandlers/CompressionComponents/Oodle/Source/OodleHandlerComponent/Private/OodleUtils.cpp b/Engine/Plugins/Runtime/PacketHandlers/CompressionComponents/Oodle/Source/OodleHandlerComponent/Private/OodleUtils.cpp
--- a/Engine/Plugins/Runtime/PacketHandlers/CompressionComponents/Oodle/Source/OodleHandlerComponent/Private/OodleUtils.cpp
+++ b/Engine/Plugins/Runtime/PacketHandlers/CompressionComponents/Oodle/Source/OodleHandlerComponent/Private/OodleUtils.cpp
@@ -1,7 +1,9 @@
 #include "OodleUtils.h"
 #include "Serialization/MemoryReader.h"
 #include "Serialization/MemoryWriter.h"
+#if HAS_OODLE_SDK
 #include "oodle2.h"
+#endif
 
 #if HAS_OODLE_SDK
 
diff --git a/Engine/Plugins/Runtime/PacketHandlers/CompressionComponents/Oodle/UnrealPakPlugin/UnrealPakPlugin.cpp b/Engine/Plugins/Runtime/PacketHandlers/CompressionComponents/Oodle/UnrealPakPlugin/UnrealPakPlugin.cpp
--- a/Engine/Plugins/Runtime/PacketHandlers/CompressionComponents/Oodle/UnrealPakPlugin/UnrealPakPlugin.cpp
+++ b/Engine/Plugins/Runtime/PacketHandlers/CompressionComponents/Oodle/UnrealPakPlugin/UnrealPakPlugin.cpp
@@ -16,13 +16,14 @@
 	#else
 		#error Unknown platform
 	#endif
-#elif WIN64
+#elif _WINDOWS
 	#define PLATFORM_WINDOWS 1
 #else
 	#error Unknown platform
 #endif
 
-#define CORE_API 
+#define CORE_API
+#define HAS_OODLE_SDK 1
 
 #include <stdlib.h>
 #include "../Source/OodleHandlerComponent/Private/OodleCustomCompressor.cpp"
diff --git a/Engine/Source/Runtime/Analytics/AnalyticsET/Private/IAnalyticsProviderET.cpp b/Engine/Source/Runtime/Analytics/AnalyticsET/Private/IAnalyticsProviderET.cpp
--- a/Engine/Source/Runtime/Analytics/AnalyticsET/Private/IAnalyticsProviderET.cpp
+++ b/Engine/Source/Runtime/Analytics/AnalyticsET/Private/IAnalyticsProviderET.cpp
@@ -445,8 +445,6 @@ bool FAnalyticsProviderET::StartSession(TArray<FAnalyticsEventAttribute>&& Attri
 
 	// always ensure we send a few specific attributes on session start.
 	TArray<FAnalyticsEventAttribute> AppendedAttributes(MoveTemp(Attributes));
-	// this allows mapping to ad networks attribution data
-	AppendedAttributes.Emplace(TEXT("AttributionId"), FPlatformMisc::GetUniqueAdvertisingId());
 	// we should always know what platform is hosting this session.
 	AppendedAttributes.Emplace(TEXT("Platform"), FString(FPlatformProperties::IniPlatformName()));
 
