Added 'ClientEnableMode' and 'ServerEnableMode' options for Oodle, to specify if compression should always be enabled ('AlwaysEnabled'), or whether it should only be enabled if the remote side requests it ('WhenCompressedPacketReceived'). Allows servers to make compression optional for clients.

diff --git a/Engine/Plugins/Runtime/PacketHandlers/CompressionComponents/Oodle/Source/OodleHandlerComponent/Private/OodleHandlerComponent.cpp b/Engine/Plugins/Runtime/PacketHandlers/CompressionComponents/Oodle/Source/OodleHandlerComponent/Private/OodleHandlerComponent.cpp
--- a/Engine/Plugins/Runtime/PacketHandlers/CompressionComponents/Oodle/Source/OodleHandlerComponent/Private/OodleHandlerComponent.cpp
+++ b/Engine/Plugins/Runtime/PacketHandlers/CompressionComponents/Oodle/Source/OodleHandlerComponent/Private/OodleHandlerComponent.cpp
@@ -9,6 +9,7 @@
 #include "Misc/ConfigCacheIni.h"
 #include "Misc/App.h"
 #include "Misc/EngineVersion.h"
+#include "UObject/Package.h"
 
 #include "OodleTrainerCommandlet.h"
 #include "OodleAnalytics.h"
@@ -272,6 +273,8 @@ FOodleDictionary::~FOodleDictionary()
 OodleHandlerComponent::OodleHandlerComponent()
 	: HandlerComponent(FName(TEXT("OodleHandlerComponent")))
 	, bEnableOodle(false)
+	, ServerEnableMode(EOodleEnableMode::AlwaysEnabled)
+	, ClientEnableMode(EOodleEnableMode::AlwaysEnabled)
 #if !UE_BUILD_SHIPPING || OODLE_DEV_SHIPPING
 	, InPacketLog(nullptr)
 	, OutPacketLog(nullptr)
@@ -283,6 +286,7 @@ OodleHandlerComponent::OodleHandlerComponent()
 	, bOodleAnalytics(false)
 	, ServerDictionary()
 	, ClientDictionary()
+	, bInitializedDictionaries(false)
 {
 	SetActive(true);
 
@@ -348,6 +352,29 @@ void OodleHandlerComponent::Initialize()
 		bEnableOodle = true;
 	}
 
+	FString CurEnableModeStr;
+	UEnum* EnableModeEnum = FindObject<UEnum>(ANY_PACKAGE, TEXT("EOodleEnableMode"), true);
+
+	if (GConfig->GetString(OODLE_INI_SECTION, TEXT("ServerEnableMode"), CurEnableModeStr, GEngineIni))
+	{
+		int64 EnumVal = EnableModeEnum->GetValueByNameString(CurEnableModeStr);
+
+		if (EnumVal != INDEX_NONE)
+		{
+			ServerEnableMode = (EOodleEnableMode)EnumVal;
+		}
+	}
+
+	if (GConfig->GetString(OODLE_INI_SECTION, TEXT("ClientEnableMode"), CurEnableModeStr, GEngineIni))
+	{
+		int64 EnumVal = EnableModeEnum->GetValueByNameString(CurEnableModeStr);
+
+		if (EnumVal != INDEX_NONE)
+		{
+			ClientEnableMode = (EOodleEnableMode)EnumVal;
+		}
+	}
+
 #if !UE_BUILD_SHIPPING || OODLE_DEV_SHIPPING
 	GConfig->GetBool(OODLE_INI_SECTION, TEXT("bUseDictionaryIfPresent"), bUseDictionaryIfPresent, GEngineIni);
 
@@ -377,7 +404,12 @@ void OodleHandlerComponent::Initialize()
 		}
 #endif
 
-		InitializeDictionaries();
+		EOodleEnableMode EnableMode = (Handler->Mode == Handler::Mode::Server ? ServerEnableMode : ClientEnableMode);
+
+		if (EnableMode == EOodleEnableMode::AlwaysEnabled)
+		{
+			InitializeDictionaries();
+		}
 	}
 
 	if (bEnableOodle)
@@ -415,6 +447,8 @@ void OodleHandlerComponent::InitializeDictionaries()
 	FString ClientDictionaryPath;
 	bool bGotDictionaryPath = false;
 
+	bInitializedDictionaries = true;
+
 #if (!UE_BUILD_SHIPPING || OODLE_DEV_SHIPPING) && !(PLATFORM_PS4 || PLATFORM_XBOXONE || PLATFORM_SWITCH)
 	if (bUseDictionaryIfPresent)
 	{
@@ -822,7 +856,15 @@ void OodleHandlerComponent::Incoming(FBitReader& Packet)
 		// If the packet is not compressed, no further processing is necessary
 		if (bCompressedPacket)
 		{
-			FOodleDictionary* CurDict = ((Handler->Mode == Handler::Mode::Server) ? ClientDictionary.Get() : ServerDictionary.Get());
+			const bool bIsServer = (Handler->Mode == Handler::Mode::Server);
+
+			// Lazy-loading of dictionary when EOodleEnableMode::WhenCompressedPacketReceived is active
+			if (!bInitializedDictionaries && (bIsServer ? ServerEnableMode : ClientEnableMode) == EOodleEnableMode::WhenCompressedPacketReceived)
+			{
+				InitializeDictionaries();
+			}
+
+			FOodleDictionary* CurDict = (bIsServer ? ClientDictionary.Get() : ServerDictionary.Get());
 
 			if (CurDict != nullptr)
 			{
@@ -977,13 +1019,20 @@ void OodleHandlerComponent::Outgoing(FBitWriter& Packet, FOutPacketTraits& Trait
 		static uint8 UncompressedData[MAX_OODLE_BUFFER];
 		static uint8 CompressedData[MAX_OODLE_BUFFER];
 
-		FOodleDictionary* CurDict = ((Handler->Mode == Handler::Mode::Server) ? ServerDictionary.Get() : ClientDictionary.Get());
+		const bool bIsServer = (Handler->Mode == Handler::Mode::Server);
+		FOodleDictionary* CurDict = (bIsServer ? ServerDictionary.Get() : ClientDictionary.Get());
 
-#if UE_BUILD_SHIPPING
-		if (CurDict != nullptr)
-#else
-		if (CurDict != nullptr && !bOodleCompressionDisabled)
+
+		bool bSkipCompression =
+#if !UE_BUILD_SHIPPING
+			// Allow a lack of dictionary in capture mode, or when compression is disabled
+			(CurDict == nullptr && bCaptureMode) || bOodleCompressionDisabled ||
 #endif
+			// Skip compression when we are waiting on the remote side to enable compression
+			(!bInitializedDictionaries && (bIsServer ? ServerEnableMode : ClientEnableMode) == EOodleEnableMode::WhenCompressedPacketReceived)
+			;
+
+		if (CurDict != nullptr && !bSkipCompression)
 		{
 #if !UE_BUILD_SHIPPING
 			check(MaxOutgoingBits <= (MAX_OODLE_PACKET_BYTES * 8));
@@ -1134,9 +1183,7 @@ void OodleHandlerComponent::Outgoing(FBitWriter& Packet, FOutPacketTraits& Trait
 				Packet.SetError();
 			}
 		}
-#if !UE_BUILD_SHIPPING || OODLE_DEV_SHIPPING
-		// Allow a lack of dictionary in capture mode, or when compression is disabled
-		else if ((CurDict == nullptr && bCaptureMode) || bOodleCompressionDisabled)
+		else if (bSkipCompression)
 		{
 			uint8 bCompresedPacket = false;
 			uint32 UncompressedBits = Packet.GetNumBits();
@@ -1147,8 +1194,7 @@ void OodleHandlerComponent::Outgoing(FBitWriter& Packet, FOutPacketTraits& Trait
 			Packet.WriteBit(bCompresedPacket);
 			Packet.SerializeBits(UncompressedData, UncompressedBits);
 		}
-#endif
-		else
+		else // if (CurDict == nullptr)
 		{
 			LowLevelFatalError(TEXT("Tried to compress a packet, but no dictionary is present for compression."));
 
@@ -1179,6 +1225,7 @@ void OodleHandlerComponent::NotifyAnalyticsProvider()
 	bOodleAnalytics = NetAnalyticsData.IsValid();
 }
 
+
 #if !UE_BUILD_SHIPPING
 /**
  * Exec Interface
diff --git a/Engine/Plugins/Runtime/PacketHandlers/CompressionComponents/Oodle/Source/OodleHandlerComponent/Public/OodleHandlerComponent.h b/Engine/Plugins/Runtime/PacketHandlers/CompressionComponents/Oodle/Source/OodleHandlerComponent/Public/OodleHandlerComponent.h
--- a/Engine/Plugins/Runtime/PacketHandlers/CompressionComponents/Oodle/Source/OodleHandlerComponent/Public/OodleHandlerComponent.h
+++ b/Engine/Plugins/Runtime/PacketHandlers/CompressionComponents/Oodle/Source/OodleHandlerComponent/Public/OodleHandlerComponent.h
@@ -21,6 +21,17 @@ DECLARE_LOG_CATEGORY_EXTERN(OodleHandlerComponentLog, Log, All);
 #include "OodleAnalytics.h"
 #include "OodleArchives.h"
 
+/**
+ * Specifies when compression is enabled. Used to make compression optional, for some platforms/clients
+ */
+UENUM()
+enum class EOodleEnableMode : uint8
+{
+	AlwaysEnabled,					// Oodle compression is always enabled - forces compression to be enabled remotely
+	WhenCompressedPacketReceived	// Oodle compression is only enabled if remotely requested
+};
+
+
 #if HAS_OODLE_SDK
 
 #if UE4_OODLE_VER >= 200
@@ -356,10 +367,17 @@ public:
 
 	virtual void NotifyAnalyticsProvider() override;
 
+
 protected:
-	/** Whether or not Oodle is enabled */
+	/** Whether or not Oodle, and its additions to the packet protocol, are enabled */
 	bool bEnableOodle;
 
+	/** When to enable compression on the server */
+	EOodleEnableMode ServerEnableMode;
+
+	/** When to enable compression on the client */
+	EOodleEnableMode ClientEnableMode;
+
 #if !UE_BUILD_SHIPPING || OODLE_DEV_SHIPPING
 	/** File to log input packets to */
 	FPacketCaptureArchive* InPacketLog;
@@ -392,6 +410,9 @@ public:
 
 	/** Client (Incoming - relative to server) dictionary data */
 	TSharedPtr<FOodleDictionary> ClientDictionary;
+
+	/** Whether or not InitializeDictionaries was ever called */
+	bool bInitializedDictionaries;
 };
 #endif
 
