Removed unwanted ensure, and restricted net analytics aggregation hotfixes to only allow enabling during runtime - disable hotfixes have to wait for the next match.

diff --git a/Engine/Plugins/Runtime/PacketHandlers/CompressionComponents/Oodle/Source/OodleHandlerComponent/Private/OodleHandlerComponent.cpp b/Engine/Plugins/Runtime/PacketHandlers/CompressionComponents/Oodle/Source/OodleHandlerComponent/Private/OodleHandlerComponent.cpp
--- a/Engine/Plugins/Runtime/PacketHandlers/CompressionComponents/Oodle/Source/OodleHandlerComponent/Private/OodleHandlerComponent.cpp
+++ b/Engine/Plugins/Runtime/PacketHandlers/CompressionComponents/Oodle/Source/OodleHandlerComponent/Private/OodleHandlerComponent.cpp
@@ -1211,21 +1211,19 @@ int32 OodleHandlerComponent::GetReservedPacketBits() const
 
 void OodleHandlerComponent::NotifyAnalyticsProvider()
 {
-	TSharedPtr<FNetAnalyticsAggregator> Aggregator = Handler->GetAggregator();
-
-	if (Handler->GetProvider().IsValid() && Aggregator.IsValid())
-	{
-		NetAnalyticsData = REGISTER_NET_ANALYTICS(Aggregator, FOodleNetAnalyticsData, TEXT("Oodle.Stats"));
-	}
-	else
+	if (!NetAnalyticsData.IsValid())
 	{
-		NetAnalyticsData.Reset();
+		TSharedPtr<FNetAnalyticsAggregator> Aggregator = Handler->GetAggregator();
+
+		if (Handler->GetProvider().IsValid() && Aggregator.IsValid())
+		{
+			NetAnalyticsData = REGISTER_NET_ANALYTICS(Aggregator, FOodleNetAnalyticsData, TEXT("Oodle.Stats"));
+		}
 	}
 
 	bOodleAnalytics = NetAnalyticsData.IsValid();
 }
 
-
 #if !UE_BUILD_SHIPPING
 /**
  * Exec Interface
diff --git a/Engine/Source/Runtime/PacketHandlers/PacketHandler/Private/PacketHandler.cpp b/Engine/Source/Runtime/PacketHandlers/PacketHandler/Private/PacketHandler.cpp
--- a/Engine/Source/Runtime/PacketHandlers/PacketHandler/Private/PacketHandler.cpp
+++ b/Engine/Source/Runtime/PacketHandlers/PacketHandler/Private/PacketHandler.cpp
@@ -146,10 +146,6 @@ void PacketHandler::NotifyAnalyticsProvider(TSharedPtr<IAnalyticsProvider> InPro
 
 	if (State != Handler::State::Uninitialized)
 	{
-		// Hotfixes should never reach this code from NetConnection's, but we can't avoid it in the case of the stateless connect handler.
-		// The latter should be ok without special/expensive multithreaded handling, as the hotfix happens so early - but it's not ideal.
-		ensure(bConnectionlessHandler || HandlerComponents.Num() == 0);
-
 		for (const TSharedPtr<HandlerComponent>& CurComponent : HandlerComponents)
 		{
 			if (CurComponent->IsInitialized())
