exposed state of packet compression in OodleHandleComponent

diff --git a/Engine/Plugins/Runtime/PacketHandlers/CompressionComponents/Oodle/Source/OodleHandlerComponent/Private/OodleHandlerComponent.cpp b/Engine/Plugins/Runtime/PacketHandlers/CompressionComponents/Oodle/Source/OodleHandlerComponent/Private/OodleHandlerComponent.cpp
--- a/Engine/Plugins/Runtime/PacketHandlers/CompressionComponents/Oodle/Source/OodleHandlerComponent/Private/OodleHandlerComponent.cpp
+++ b/Engine/Plugins/Runtime/PacketHandlers/CompressionComponents/Oodle/Source/OodleHandlerComponent/Private/OodleHandlerComponent.cpp
@@ -1204,6 +1204,33 @@ void OodleHandlerComponent::Outgoing(FBitWriter& Packet, FOutPacketTraits& Trait
 	}
 }
 
+bool OodleHandlerComponent::IsCompressionActive() const
+{
+	bool bResult = false;
+
+	if (bEnableOodle && Handler)
+	{
+		const bool bIsServer = (Handler->Mode == Handler::Mode::Server);
+		FOodleDictionary* CurDict = (bIsServer ? ServerDictionary.Get() : ClientDictionary.Get());
+
+		const bool bSkipCompression =
+#if !UE_BUILD_SHIPPING
+			// Allow a lack of dictionary in capture mode, or when compression is disabled
+			(CurDict == nullptr && bCaptureMode) || bOodleCompressionDisabled ||
+#endif
+			// Skip compression when we are waiting on the remote side to enable compression
+			(!bInitializedDictionaries && (bIsServer ? ServerEnableMode : ClientEnableMode) == EOodleEnableMode::WhenCompressedPacketReceived)
+			;
+
+		if (CurDict != nullptr && !bSkipCompression)
+		{
+			bResult = true;
+		}
+	}
+
+	return bResult;
+}
+
 int32 OodleHandlerComponent::GetReservedPacketBits() const
 {
 	return bEnableOodle ? OodleReservedPacketBits : 0;
diff --git a/Engine/Plugins/Runtime/PacketHandlers/CompressionComponents/Oodle/Source/OodleHandlerComponent/Public/OodleHandlerComponent.h b/Engine/Plugins/Runtime/PacketHandlers/CompressionComponents/Oodle/Source/OodleHandlerComponent/Public/OodleHandlerComponent.h
--- a/Engine/Plugins/Runtime/PacketHandlers/CompressionComponents/Oodle/Source/OodleHandlerComponent/Public/OodleHandlerComponent.h
+++ b/Engine/Plugins/Runtime/PacketHandlers/CompressionComponents/Oodle/Source/OodleHandlerComponent/Public/OodleHandlerComponent.h
@@ -347,6 +347,13 @@ public:
 	void FreePacketLogs();
 #endif
 
+	/** 
+	 * Check if component is currently compressing packets
+	 *
+	 * @return	Whether or not compression is active
+	 */
+	bool IsCompressionActive() const;
+
 	virtual void Initialize() override;
 
 	virtual bool IsValid() const override;
