Added CVar allowing Oodle to skip compressing packets below a particular size. Improved Oodle Analytics, to factor in skipped packets.

If tuned to skip packets that fit within a single ethernet frame (CVar value ~18), will reduce CPU usage while (theoretically) keeping the same network performance (on-the-wire bandwidth stays the same, but bandwidth calculations for AWS and mobile data will be affected).

diff --git a/Engine/Plugins/Runtime/PacketHandlers/CompressionComponents/Oodle/Source/OodleHandlerComponent/Private/OodleAnalytics.cpp b/Engine/Plugins/Runtime/PacketHandlers/CompressionComponents/Oodle/Source/OodleHandlerComponent/Private/OodleAnalytics.cpp
--- a/Engine/Plugins/Runtime/PacketHandlers/CompressionComponents/Oodle/Source/OodleHandlerComponent/Private/OodleAnalytics.cpp
+++ b/Engine/Plugins/Runtime/PacketHandlers/CompressionComponents/Oodle/Source/OodleHandlerComponent/Private/OodleAnalytics.cpp
@@ -18,15 +18,20 @@ FOodleAnalyticsVars::FOodleAnalyticsVars()
 	, InCompressedWithOverheadLengthTotal(0)
 	, InCompressedLengthTotal(0)
 	, InDecompressedLengthTotal(0)
+	, InNotCompressedLengthTotal(0)
 	, OutCompressedNum(0)
 	, OutNotCompressedFailedNum(0)
-	, OutNotCompressedBoundedNum(0)
-	, OutNotCompressedFlaggedNum(0)
 	, OutNotCompressedFailedAckOnlyNum(0)
 	, OutNotCompressedFailedKeepAliveNum(0)
+	, OutNotCompressedBoundedNum(0)
+	, OutNotCompressedFlaggedNum(0)
+	, OutNotCompressedClientDisabledNum(0)
+	, OutNotCompressedTooSmallNum(0)
 	, OutCompressedWithOverheadLengthTotal(0)
 	, OutCompressedLengthTotal(0)
 	, OutBeforeCompressedLengthTotal(0)
+	, OutNotCompressedFailedLengthTotal(0)
+	, OutNotCompressedSkippedLengthTotal(0)
 {
 }
 
@@ -37,15 +42,20 @@ bool FOodleAnalyticsVars::operator == (const FOodleAnalyticsVars& A) const
 		A.InCompressedWithOverheadLengthTotal == InCompressedWithOverheadLengthTotal &&
 		A.InCompressedLengthTotal == InCompressedLengthTotal &&
 		A.InDecompressedLengthTotal == InDecompressedLengthTotal &&
+		A.InNotCompressedLengthTotal == InNotCompressedLengthTotal &&
 		A.OutCompressedNum == OutCompressedNum &&
 		A.OutNotCompressedFailedNum == OutNotCompressedFailedNum &&
-		A.OutNotCompressedBoundedNum == OutNotCompressedBoundedNum &&
-		A.OutNotCompressedFlaggedNum == OutNotCompressedFlaggedNum &&
 		A.OutNotCompressedFailedAckOnlyNum == OutNotCompressedFailedAckOnlyNum &&
 		A.OutNotCompressedFailedKeepAliveNum == OutNotCompressedFailedKeepAliveNum &&
+		A.OutNotCompressedBoundedNum == OutNotCompressedBoundedNum &&
+		A.OutNotCompressedFlaggedNum == OutNotCompressedFlaggedNum &&
+		A.OutNotCompressedClientDisabledNum == OutNotCompressedClientDisabledNum &&
+		A.OutNotCompressedTooSmallNum == OutNotCompressedTooSmallNum &&
 		A.OutCompressedWithOverheadLengthTotal == OutCompressedWithOverheadLengthTotal &&
 		A.OutCompressedLengthTotal == OutCompressedLengthTotal &&
-		A.OutBeforeCompressedLengthTotal == OutBeforeCompressedLengthTotal;
+		A.OutBeforeCompressedLengthTotal == OutBeforeCompressedLengthTotal &&
+		A.OutNotCompressedFailedLengthTotal == OutNotCompressedFailedLengthTotal &&
+		A.OutNotCompressedSkippedLengthTotal == OutNotCompressedSkippedLengthTotal;
 }
 
 void FOodleAnalyticsVars::CommitAnalytics(FOodleAnalyticsVars& AggregatedData)
@@ -55,15 +65,20 @@ void FOodleAnalyticsVars::CommitAnalytics(FOodleAnalyticsVars& AggregatedData)
 	AggregatedData.InCompressedWithOverheadLengthTotal += InCompressedWithOverheadLengthTotal;
 	AggregatedData.InCompressedLengthTotal += InCompressedLengthTotal;
 	AggregatedData.InDecompressedLengthTotal += InDecompressedLengthTotal;
+	AggregatedData.InNotCompressedLengthTotal += InNotCompressedLengthTotal;
 	AggregatedData.OutCompressedNum += OutCompressedNum;
 	AggregatedData.OutNotCompressedFailedNum += OutNotCompressedFailedNum;
-	AggregatedData.OutNotCompressedBoundedNum += OutNotCompressedBoundedNum;
-	AggregatedData.OutNotCompressedFlaggedNum += OutNotCompressedFlaggedNum;
 	AggregatedData.OutNotCompressedFailedAckOnlyNum += OutNotCompressedFailedAckOnlyNum;
 	AggregatedData.OutNotCompressedFailedKeepAliveNum += OutNotCompressedFailedKeepAliveNum;
+	AggregatedData.OutNotCompressedBoundedNum += OutNotCompressedBoundedNum;
+	AggregatedData.OutNotCompressedFlaggedNum += OutNotCompressedFlaggedNum;
+	AggregatedData.OutNotCompressedClientDisabledNum += OutNotCompressedClientDisabledNum;
+	AggregatedData.OutNotCompressedTooSmallNum = OutNotCompressedTooSmallNum;
 	AggregatedData.OutCompressedWithOverheadLengthTotal += OutCompressedWithOverheadLengthTotal;
 	AggregatedData.OutCompressedLengthTotal += OutCompressedLengthTotal;
 	AggregatedData.OutBeforeCompressedLengthTotal += OutBeforeCompressedLengthTotal;
+	AggregatedData.OutNotCompressedFailedLengthTotal += OutNotCompressedFailedLengthTotal;
+	AggregatedData.OutNotCompressedSkippedLengthTotal += OutNotCompressedSkippedLengthTotal;
 }
 
 
@@ -80,38 +95,66 @@ void FOodleNetAnalyticsData::SendAnalytics()
 	if (!(*this == NullVars) && AnalyticsProvider.IsValid())
 	{
 		/** The number of outgoing packets that were not compressed, in total */
-		uint64 OutNotCompressedNumTotal = OutNotCompressedFailedNum + OutNotCompressedBoundedNum + OutNotCompressedFlaggedNum;
+		uint64 OutNotCompressedNumTotal = OutNotCompressedFailedNum + OutNotCompressedBoundedNum + OutNotCompressedFlaggedNum + OutNotCompressedClientDisabledNum + OutNotCompressedTooSmallNum;
+
+		/** The number of outgoing packets that had compression skipped, in total */
+		uint64 OutNotCompressedSkippedNum = OutNotCompressedFlaggedNum + OutNotCompressedClientDisabledNum + OutNotCompressedTooSmallNum;
+
+
+		uint64 InPreLengthTotal = InCompressedLengthTotal + InNotCompressedLengthTotal;
+		uint64 InPreWithOverheadLengthTotal = InCompressedWithOverheadLengthTotal + InNotCompressedLengthTotal;
+		uint64 InPostLengthTotal = InDecompressedLengthTotal + InNotCompressedLengthTotal;
+
+		uint64 OutPreLengthTotal = OutBeforeCompressedLengthTotal + OutNotCompressedFailedLengthTotal + OutNotCompressedSkippedLengthTotal;
+		uint64 OutPostLengthTotal = OutCompressedLengthTotal + OutNotCompressedFailedLengthTotal + OutNotCompressedSkippedLengthTotal;
+		uint64 OutPostWithOverheadLengthTotal = OutCompressedWithOverheadLengthTotal + OutNotCompressedFailedLengthTotal + OutNotCompressedSkippedLengthTotal;
+
+		uint64 OutPreAttemptedLengthTotal = OutBeforeCompressedLengthTotal + OutNotCompressedFailedLengthTotal;
+		uint64 OutPostAttemptedWithOverheadLengthTotal = OutCompressedWithOverheadLengthTotal + OutNotCompressedFailedLengthTotal;
+
 
+		// @todo #JohnB: Deprecate these four algorithm-based analytics, eventually - since skipped packets render them inaccurate
 		/**
 		 * The below values measure Oodle algorithm compression, minus overhead reducing final savings.
+		 * Also factors in skipped/failed compression (which has increased a lot since analytics was added) - reducing the usefulness for determining algorithm compression.
 		 */
 			/** The percentage of compression savings, of all incoming packets. */
-			int8 InSavingsPercentTotal = (1.0 - ((double)InCompressedLengthTotal / (double)InDecompressedLengthTotal)) * 100.0;
+			int8 InSavingsPercentTotal = (1.0 - ((double)InPreLengthTotal / (double)InPostLengthTotal)) * 100.0;
 
 			/** The percentage of compression savings, of all outgoing packets. */
-			int8 OutSavingsPercentTotal = (1.0 - ((double)OutCompressedLengthTotal / (double)OutBeforeCompressedLengthTotal)) * 100.0;
+			int8 OutSavingsPercentTotal = (1.0 - ((double)OutPostLengthTotal / (double)OutPreLengthTotal)) * 100.0;
 
 			/** The number of bytes saved due to compression, of all incoming packets. */
-			int64 InSavingsBytesTotal = InDecompressedLengthTotal - InCompressedLengthTotal;
+			int64 InSavingsBytesTotal = InPostLengthTotal - InPreLengthTotal;
 
 			/** The number of bytes saved due to compression, of all outgoing packets. */
-			int64 OutSavingsBytesTotal = OutBeforeCompressedLengthTotal - OutCompressedLengthTotal;
+			int64 OutSavingsBytesTotal = OutPreLengthTotal - OutPostLengthTotal;
 
 		/**
 		 * The below values measure compressed length + decompression data overhead, which reduces final savings.
-		 * This is the most accurate measure of compression savings.
+		 * This is the most accurate measure of compression savings (in terms of overall bandwidth).
 		 */
 			/** The percentage of compression savings, of all incoming packets. */
-			int8 InSavingsWithOverheadPercentTotal = (1.0 - ((double)InCompressedWithOverheadLengthTotal / (double)InDecompressedLengthTotal)) * 100.0;
+			int8 InSavingsWithOverheadPercentTotal = (1.0 - ((double)InPreWithOverheadLengthTotal / (double)InPostLengthTotal)) * 100.0;
 
 			/** The percentage of compression savings, of all outgoing packets. */
-			int8 OutSavingsWithOverheadPercentTotal = (1.0 - ((double)OutCompressedWithOverheadLengthTotal / (double)OutBeforeCompressedLengthTotal)) * 100.0;
+			int8 OutSavingsWithOverheadPercentTotal = (1.0 - ((double)OutPostWithOverheadLengthTotal / (double)OutPreLengthTotal)) * 100.0;
 
 			/** The number of bytes saved due to compression, of all incoming packets. */
-			int64 InSavingsWithOverheadBytesTotal = InDecompressedLengthTotal - InCompressedWithOverheadLengthTotal;
+			int64 InSavingsWithOverheadBytesTotal = InPostLengthTotal - InPreWithOverheadLengthTotal;
 
 			/** The number of bytes saved due to compression, of all outgoing packets. */
-			int64 OutSavingsWithOverheadBytesTotal = OutBeforeCompressedLengthTotal - OutCompressedWithOverheadLengthTotal;
+			int64 OutSavingsWithOverheadBytesTotal = OutPreLengthTotal - OutPostWithOverheadLengthTotal;
+
+		/**
+		 * The below values measure compressed length + decompression data overhead, only for packets where compression was attempted (i.e. ignores skipped packets).
+		 * This is the best measure for determining algorithm compression performance, especially when measured against the Oodle encode/decode CPU cost.
+		 */
+			/** The percentage of compression savings, of all incoming packets that were compressed (attempted but failed compress, can't be counted here) */
+			int8 InAttemptedSavingsWithOverheadPercentTotal = (1.0 - ((double)InCompressedWithOverheadLengthTotal / (double)InDecompressedLengthTotal)) * 100.0;
+
+			/** The percentage of compression savings, of all outgoing packets that attempted compression. */
+			int8 OutAttemptedSavingsWithOverheadPercentTotal = (1.0 - ((double)OutPostAttemptedWithOverheadLengthTotal / (double)OutPreAttemptedLengthTotal)) * 100.0;
 
 
 		UE_LOG(OodleHandlerComponentLog, Log, TEXT("Oodle Analytics:"));
@@ -120,15 +163,21 @@ void FOodleNetAnalyticsData::SendAnalytics()
 		UE_LOG(OodleHandlerComponentLog, Log, TEXT(" - InCompressedWithOverheadLengthTotal: %llu"), InCompressedWithOverheadLengthTotal);
 		UE_LOG(OodleHandlerComponentLog, Log, TEXT(" - InCompressedLengthTotal: %llu"), InCompressedLengthTotal);
 		UE_LOG(OodleHandlerComponentLog, Log, TEXT(" - InDecompressedLengthTotal: %llu"), InDecompressedLengthTotal);
+		UE_LOG(OodleHandlerComponentLog, Log, TEXT(" - InNotCompressedLengthTotal: %llu"), InNotCompressedLengthTotal);
 		UE_LOG(OodleHandlerComponentLog, Log, TEXT(" - OutCompressedNum: %llu"), OutCompressedNum);
 		UE_LOG(OodleHandlerComponentLog, Log, TEXT(" - OutNotCompressedFailedNum: %llu"), OutNotCompressedFailedNum);
-		UE_LOG(OodleHandlerComponentLog, Log, TEXT(" - OutNotCompressedBoundedNum: %llu"), OutNotCompressedBoundedNum);
-		UE_LOG(OodleHandlerComponentLog, Log, TEXT(" - OutNotCompressedFlaggedNum: %llu"), OutNotCompressedFlaggedNum);
 		UE_LOG(OodleHandlerComponentLog, Log, TEXT(" - OutNotCompressedFailedAckOnlyNum: %llu"), OutNotCompressedFailedAckOnlyNum);
 		UE_LOG(OodleHandlerComponentLog, Log, TEXT(" - OutNotCompressedFailedKeepAliveNum: %llu"), OutNotCompressedFailedKeepAliveNum);
+		UE_LOG(OodleHandlerComponentLog, Log, TEXT(" - OutNotCompressedBoundedNum: %llu"), OutNotCompressedBoundedNum);
+		UE_LOG(OodleHandlerComponentLog, Log, TEXT(" - OutNotCompressedFlaggedNum: %llu"), OutNotCompressedFlaggedNum);
+		UE_LOG(OodleHandlerComponentLog, Log, TEXT(" - OutNotCompressedSkippedNum: %llu"), OutNotCompressedSkippedNum);
+		UE_LOG(OodleHandlerComponentLog, Log, TEXT(" - OutNotCompressedClientDisabledNum: %llu"), OutNotCompressedClientDisabledNum);
+		UE_LOG(OodleHandlerComponentLog, Log, TEXT(" - OutNotCompressedTooSmallNum: %llu"), OutNotCompressedTooSmallNum);
 		UE_LOG(OodleHandlerComponentLog, Log, TEXT(" - OutCompressedWithOverheadLengthTotal: %llu"), OutCompressedWithOverheadLengthTotal);
 		UE_LOG(OodleHandlerComponentLog, Log, TEXT(" - OutCompressedLengthTotal: %llu"), OutCompressedLengthTotal);
 		UE_LOG(OodleHandlerComponentLog, Log, TEXT(" - OutBeforeCompressedLengthTotal: %llu"), OutBeforeCompressedLengthTotal);
+		UE_LOG(OodleHandlerComponentLog, Log, TEXT(" - OutNotCompressedFailedLengthTotal: %llu"), OutNotCompressedFailedLengthTotal);
+		UE_LOG(OodleHandlerComponentLog, Log, TEXT(" - OutNotCompressedSkippedLengthTotal: %llu"), OutNotCompressedSkippedLengthTotal);
 		UE_LOG(OodleHandlerComponentLog, Log, TEXT(" - OutNotCompressedNumTotal: %llu"), OutNotCompressedNumTotal);
 		UE_LOG(OodleHandlerComponentLog, Log, TEXT(" - InSavingsPercentTotal: %i"), InSavingsPercentTotal);
 		UE_LOG(OodleHandlerComponentLog, Log, TEXT(" - OutSavingsPercentTotal: %i"), OutSavingsPercentTotal);
@@ -138,6 +187,8 @@ void FOodleNetAnalyticsData::SendAnalytics()
 		UE_LOG(OodleHandlerComponentLog, Log, TEXT(" - OutSavingsWithOverheadPercentTotal: %i"), OutSavingsWithOverheadPercentTotal);
 		UE_LOG(OodleHandlerComponentLog, Log, TEXT(" - InSavingsWithOverheadBytesTotal: %lli"), InSavingsWithOverheadBytesTotal);
 		UE_LOG(OodleHandlerComponentLog, Log, TEXT(" - OutSavingsWithOverheadBytesTotal: %lli"), OutSavingsWithOverheadBytesTotal);
+		UE_LOG(OodleHandlerComponentLog, Log, TEXT(" - InAttemptedSavingsWithOverheadPercentTotal: %i"), InAttemptedSavingsWithOverheadPercentTotal);
+		UE_LOG(OodleHandlerComponentLog, Log, TEXT(" - OutAttemptedSavingsWithOverheadPercentTotal: %i"), OutAttemptedSavingsWithOverheadPercentTotal);
 
 
 		static const FString EZEventName = TEXT("Oodle.Stats");
@@ -146,15 +197,21 @@ void FOodleNetAnalyticsData::SendAnalytics()
 		static const FString EZAttrib_InCompressedWithOverheadLengthTotal = TEXT("InCompressedWithOverheadLengthTotal");
 		static const FString EZAttrib_InCompressedLengthTotal = TEXT("InCompressedLengthTotal");
 		static const FString EZAttrib_InDecompressedLengthTotal = TEXT("InDecompressedLengthTotal");
+		static const FString EZAttrib_InNotCompressedLengthTotal = TEXT("InNotCompressedLengthTotal");
 		static const FString EZAttrib_OutCompressedNum = TEXT("OutCompressedNum");
 		static const FString EZAttrib_OutNotCompressedFailedNum = TEXT("OutNotCompressedFailedNum");
-		static const FString EZAttrib_OutNotCompressedBoundedNum = TEXT("OutNotCompressedBoundedNum");
-		static const FString EZAttrib_OutNotCompressedFlaggedNum = TEXT("OutNotCompressedFlaggedNum");
 		static const FString EZAttrib_OutNotCompressedFailedAckOnlyNum = TEXT("OutNotCompressedFailedAckOnlyNum");
 		static const FString EZAttrib_OutNotCompressedFailedKeepAliveNum = TEXT("OutNotCompressedFailedKeepAliveNum");
+		static const FString EZAttrib_OutNotCompressedBoundedNum = TEXT("OutNotCompressedBoundedNum");
+		static const FString EZAttrib_OutNotCompressedFlaggedNum = TEXT("OutNotCompressedFlaggedNum");
+		static const FString EZAttrib_OutNotCompressedSkippedNum = TEXT("OutNotCompressedSkippedNum");
+		static const FString EZAttrib_OutNotCompressedClientDisabledNum = TEXT("OutNotCompressedClientDisabledNum");
+		static const FString EZAttrib_OutNotCompressedTooSmallNum = TEXT("OutNotCompressedTooSmallNum");
 		static const FString EZAttrib_OutCompressedWithOverheadLengthTotal = TEXT("OutCompressedWithOverheadLengthTotal");
 		static const FString EZAttrib_OutCompressedLengthTotal = TEXT("OutCompressedLengthTotal");
 		static const FString EZAttrib_OutBeforeCompressedLengthTotal = TEXT("OutBeforeCompressedLengthTotal");
+		static const FString EZAttrib_OutNotCompressedFailedLengthTotal = TEXT("OutNotCompressedFailedLengthTotal");
+		static const FString EZAttrib_OutNotCompressedSkippedLengthTotal = TEXT("OutNotCompressedSkippedLengthTotal");
 		static const FString EZAttrib_OutNotCompressedNumTotal = TEXT("OutNotCompressedNumTotal");
 		static const FString EZAttrib_InSavingsPercentTotal = TEXT("InSavingsPercentTotal");
 		static const FString EZAttrib_OutSavingsPercentTotal = TEXT("OutSavingsPercentTotal");
@@ -164,22 +221,30 @@ void FOodleNetAnalyticsData::SendAnalytics()
 		static const FString EZAttrib_OutSavingsWithOverheadPercentTotal = TEXT("OutSavingsWithOverheadPercentTotal");
 		static const FString EZAttrib_InSavingsWithOverheadBytesTotal = TEXT("InSavingsWithOverheadBytesTotal");
 		static const FString EZAttrib_OutSavingsWithOverheadBytesTotal = TEXT("OutSavingsWithOverheadBytesTotal");
+		static const FString EZAttrib_InAttemptedSavingsWithOverheadPercentTotal = TEXT("InAttemptedSavingsWithOverheadPercentTotal");
+		static const FString EZAttrib_OutAttemptedSavingsWithOverheadPercentTotal = TEXT("OutAttemptedSavingsWithOverheadPercentTotal");
 
-		AnalyticsProvider->RecordEvent(EZEventName, MakeAnalyticsEventAttributeArray(
+		const TArray<FAnalyticsEventAttribute> EventAttributes = MakeAnalyticsEventAttributeArray(
 			EZAttrib_InCompressedNum, InCompressedNum,
 			EZAttrib_InNotCompressedNum, InNotCompressedNum,
 			EZAttrib_InCompressedWithOverheadLengthTotal, InCompressedWithOverheadLengthTotal,
 			EZAttrib_InCompressedLengthTotal, InCompressedLengthTotal,
 			EZAttrib_InDecompressedLengthTotal, InDecompressedLengthTotal,
+			EZAttrib_InNotCompressedLengthTotal, InNotCompressedLengthTotal,
 			EZAttrib_OutCompressedNum, OutCompressedNum,
 			EZAttrib_OutNotCompressedFailedNum, OutNotCompressedFailedNum,
-			EZAttrib_OutNotCompressedBoundedNum, OutNotCompressedBoundedNum,
-			EZAttrib_OutNotCompressedFlaggedNum, OutNotCompressedFlaggedNum,
 			EZAttrib_OutNotCompressedFailedAckOnlyNum, OutNotCompressedFailedAckOnlyNum,
 			EZAttrib_OutNotCompressedFailedKeepAliveNum, OutNotCompressedFailedKeepAliveNum,
+			EZAttrib_OutNotCompressedBoundedNum, OutNotCompressedBoundedNum,
+			EZAttrib_OutNotCompressedFlaggedNum, OutNotCompressedFlaggedNum,
+			EZAttrib_OutNotCompressedSkippedNum, OutNotCompressedSkippedNum,
+			EZAttrib_OutNotCompressedClientDisabledNum, OutNotCompressedClientDisabledNum,
+			EZAttrib_OutNotCompressedTooSmallNum, OutNotCompressedTooSmallNum,
 			EZAttrib_OutCompressedWithOverheadLengthTotal, OutCompressedWithOverheadLengthTotal,
 			EZAttrib_OutCompressedLengthTotal, OutCompressedLengthTotal,
 			EZAttrib_OutBeforeCompressedLengthTotal, OutBeforeCompressedLengthTotal,
+			EZAttrib_OutNotCompressedFailedLengthTotal, OutNotCompressedFailedLengthTotal,
+			EZAttrib_OutNotCompressedSkippedLengthTotal, OutNotCompressedSkippedLengthTotal,
 			EZAttrib_OutNotCompressedNumTotal, OutNotCompressedNumTotal,
 			EZAttrib_InSavingsPercentTotal, InSavingsPercentTotal,
 			EZAttrib_OutSavingsPercentTotal, OutSavingsPercentTotal,
@@ -188,7 +253,11 @@ void FOodleNetAnalyticsData::SendAnalytics()
 			EZAttrib_InSavingsWithOverheadPercentTotal, InSavingsWithOverheadPercentTotal,
 			EZAttrib_OutSavingsWithOverheadPercentTotal, OutSavingsWithOverheadPercentTotal,
 			EZAttrib_InSavingsWithOverheadBytesTotal, InSavingsWithOverheadBytesTotal,
-			EZAttrib_OutSavingsWithOverheadBytesTotal, OutSavingsWithOverheadBytesTotal
-		));
+			EZAttrib_OutSavingsWithOverheadBytesTotal, OutSavingsWithOverheadBytesTotal,
+			EZAttrib_InAttemptedSavingsWithOverheadPercentTotal, InAttemptedSavingsWithOverheadPercentTotal,
+			EZAttrib_OutAttemptedSavingsWithOverheadPercentTotal, OutAttemptedSavingsWithOverheadPercentTotal
+		);
+
+		AnalyticsProvider->RecordEvent(EZEventName, EventAttributes);
 	}
 }
diff --git a/Engine/Plugins/Runtime/PacketHandlers/CompressionComponents/Oodle/Source/OodleHandlerComponent/Private/OodleHandlerComponent.cpp b/Engine/Plugins/Runtime/PacketHandlers/CompressionComponents/Oodle/Source/OodleHandlerComponent/Private/OodleHandlerComponent.cpp
--- a/Engine/Plugins/Runtime/PacketHandlers/CompressionComponents/Oodle/Source/OodleHandlerComponent/Private/OodleHandlerComponent.cpp
+++ b/Engine/Plugins/Runtime/PacketHandlers/CompressionComponents/Oodle/Source/OodleHandlerComponent/Private/OodleHandlerComponent.cpp
@@ -89,6 +89,17 @@ static TArray<OodleHandlerComponent*> OodleComponentList;
 #endif
 
 
+
+/**
+ * CVars
+ */
+
+TAutoConsoleVariable<int32> CVarOodleMinSizeForCompression(
+	TEXT("net.OodleMinSizeForCompression"), 0,
+	TEXT("The minimum size an outgoing packet must be, for it to be considered for compression (does not count overhead of handler components which process packets after Oodle)."));
+
+
+
 // @todo #JohnB: Remove after Oodle update, and after checking with Luigi
 static rrbool STDCALL UEOodleDisplayAssert(const char* File, const int Line, const char* Function, const char* Message)
 {
@@ -977,9 +988,7 @@ void OodleHandlerComponent::Incoming(FBitReader& Packet)
 				uint32 PacketSize = Packet.GetBytesLeft();
 
 				AnalyticsVars->InNotCompressedNum++;
-				AnalyticsVars->InCompressedWithOverheadLengthTotal += PacketSize;
-				AnalyticsVars->InCompressedLengthTotal += PacketSize;
-				AnalyticsVars->InDecompressedLengthTotal += PacketSize;
+				AnalyticsVars->InNotCompressedLengthTotal += PacketSize;
 			}
 
 #if !UE_BUILD_SHIPPING || OODLE_DEV_SHIPPING
@@ -1019,18 +1028,33 @@ void OodleHandlerComponent::Outgoing(FBitWriter& Packet, FOutPacketTraits& Trait
 		static uint8 UncompressedData[MAX_OODLE_BUFFER];
 		static uint8 CompressedData[MAX_OODLE_BUFFER];
 
+		FOodleAnalyticsVars* AnalyticsVars = (bOodleAnalytics && NetAnalyticsData.IsValid()) ? NetAnalyticsData->GetLocalData() : nullptr;
 		const bool bIsServer = (Handler->Mode == Handler::Mode::Server);
 		FOodleDictionary* CurDict = (bIsServer ? ServerDictionary.Get() : ClientDictionary.Get());
+		uint32 MinSizeForCompression = CVarOodleMinSizeForCompression.GetValueOnAnyThread();
+		uint32 UncompressedBytes = Packet.GetNumBytes();
+		bool bSkipCompressionClientDisabled = false;
+		bool bSkipCompressionTooSmall = false;
+		bool bSkipCompression = false;
 
-
-		bool bSkipCompression =
 #if !UE_BUILD_SHIPPING
-			// Allow a lack of dictionary in capture mode, or when compression is disabled
-			(CurDict == nullptr && bCaptureMode) || bOodleCompressionDisabled ||
+		// Allow a lack of dictionary in capture mode, or when compression is disabled
+		bSkipCompression = (CurDict == nullptr && bCaptureMode) || bOodleCompressionDisabled;
 #endif
-			// Skip compression when we are waiting on the remote side to enable compression
-			(!bInitializedDictionaries && (bIsServer ? ServerEnableMode : ClientEnableMode) == EOodleEnableMode::WhenCompressedPacketReceived)
-			;
+
+		// Skip compression when we are waiting on the remote side to enable compression
+		if (!bSkipCompression && (!bInitializedDictionaries && (bIsServer ? ServerEnableMode : ClientEnableMode) == EOodleEnableMode::WhenCompressedPacketReceived))
+		{
+			bSkipCompression = true;
+			bSkipCompressionClientDisabled = true;
+		}
+
+		// Skip compression when the packet is below the minimum size
+		if (!bSkipCompression && (MinSizeForCompression > 0 && UncompressedBytes < MinSizeForCompression))
+		{
+			bSkipCompression = true;
+			bSkipCompressionTooSmall = true;
+		}
 
 		if (CurDict != nullptr && !bSkipCompression)
 		{
@@ -1039,7 +1063,6 @@ void OodleHandlerComponent::Outgoing(FBitWriter& Packet, FOutPacketTraits& Trait
 #endif
 			uint32 MaxAdjustedLengthBits = MaxOutgoingBits - OodleReservedPacketBits;
 			uint32 UncompressedBits = Packet.GetNumBits();
-			uint32 UncompressedBytes = Packet.GetNumBytes();
 
 			bool bWithinBitBounds = UncompressedBits > 0 && ensure(UncompressedBits <= MaxAdjustedLengthBits) &&
 									ensure(OodleLZ_GetCompressedBufferSizeNeeded(UncompressedBytes) <= MAX_OODLE_BUFFER);
@@ -1078,14 +1101,6 @@ void OodleHandlerComponent::Outgoing(FBitWriter& Packet, FOutPacketTraits& Trait
 
 					Traits.bIsCompressed = !!bCompressedPacket;
 
-
-					FOodleAnalyticsVars* AnalyticsVars = (bOodleAnalytics && NetAnalyticsData.IsValid()) ? NetAnalyticsData->GetLocalData() : nullptr;
-
-					if (AnalyticsVars != nullptr)
-					{
-						AnalyticsVars->OutBeforeCompressedLengthTotal += UncompressedBytes;
-					}
-
 					if (bCompressedPacket)
 					{
 						// @todo #JohnB: Compress directly into a (deliberately oversized) FBitWriter buffer, which you can shrink after
@@ -1104,14 +1119,14 @@ void OodleHandlerComponent::Outgoing(FBitWriter& Packet, FOutPacketTraits& Trait
 							AnalyticsVars->OutCompressedNum++;
 							AnalyticsVars->OutCompressedWithOverheadLengthTotal += ((Packet.GetNumBits() - 1) + 7) >> 3;
 							AnalyticsVars->OutCompressedLengthTotal += CompressedBytes;
+							AnalyticsVars->OutBeforeCompressedLengthTotal += UncompressedBytes;
 						}
 					}
 					else
 					{
 						if (AnalyticsVars != nullptr)
 						{
-							AnalyticsVars->OutCompressedLengthTotal += UncompressedBytes;
-							AnalyticsVars->OutCompressedWithOverheadLengthTotal += UncompressedBytes;
+							AnalyticsVars->OutNotCompressedFailedLengthTotal += UncompressedBytes;
 						}
 
 						// @todo #JohnB: Try to eliminate this appBitscpy, by reserving the bCompressedPacket bit and other data,
@@ -1185,6 +1200,13 @@ void OodleHandlerComponent::Outgoing(FBitWriter& Packet, FOutPacketTraits& Trait
 		}
 		else if (bSkipCompression)
 		{
+			if (AnalyticsVars != nullptr)
+			{
+				AnalyticsVars->OutNotCompressedSkippedLengthTotal += UncompressedBytes;
+				AnalyticsVars->OutNotCompressedClientDisabledNum += (uint8)bSkipCompressionClientDisabled;
+				AnalyticsVars->OutNotCompressedTooSmallNum += (uint8)bSkipCompressionTooSmall;
+			}
+
 			uint8 bCompresedPacket = false;
 			uint32 UncompressedBits = Packet.GetNumBits();
 
diff --git a/Engine/Plugins/Runtime/PacketHandlers/CompressionComponents/Oodle/Source/OodleHandlerComponent/Public/OodleAnalytics.h b/Engine/Plugins/Runtime/PacketHandlers/CompressionComponents/Oodle/Source/OodleHandlerComponent/Public/OodleAnalytics.h
--- a/Engine/Plugins/Runtime/PacketHandlers/CompressionComponents/Oodle/Source/OodleHandlerComponent/Public/OodleAnalytics.h
+++ b/Engine/Plugins/Runtime/PacketHandlers/CompressionComponents/Oodle/Source/OodleHandlerComponent/Public/OodleAnalytics.h
@@ -33,6 +33,7 @@ public:
 	/** The number of incoming packets that were not compressed */
 	uint64 InNotCompressedNum;
 
+
 	/** The compressed length + decompression data overhead, of all incoming packets. The most accurate measure of compression savings. */
 	uint64 InCompressedWithOverheadLengthTotal;
 
@@ -42,23 +43,34 @@ public:
 	/** The decompressed length of all incoming packets. */
 	uint64 InDecompressedLengthTotal;
 
+	/** The length of all incoming packets, which were not compressed. */
+	uint64 InNotCompressedLengthTotal;
+
+
 	/** The number of outgoing compressed packets. */
 	uint64 OutCompressedNum;
 
 	/** The number of outgoing packets that were not compressed, due to Oodle failing to compress enough. */
 	uint64 OutNotCompressedFailedNum;
 
+		/** The number of outgoing packets that were not compressed, due to Oodle failing to compress - which exclusively contained ack data */
+		uint64 OutNotCompressedFailedAckOnlyNum;
+
+		/** The number of outgoing packets that were not compressed, due to Oodle failing to compress - which were KeepAlive packets */
+		uint64 OutNotCompressedFailedKeepAliveNum;
+
 	/** The number of outgoing packets that were not compressed, due to byte rounding of compressed packets, exceeding size limits. */
 	uint64 OutNotCompressedBoundedNum;
 
 	/** The number of outgoing packets that were not compressed, due to a higher level flag requesting they be sent uncompressed. */
 	uint64 OutNotCompressedFlaggedNum;
 
-	/** The number of outgoing packets that were not compressed, due to Oodle failing to compress - which exclusively contained ack data */
-	uint64 OutNotCompressedFailedAckOnlyNum;
+	/** The number of outgoing packets that were not compressed, due to the client having disabled compression. */
+	uint64 OutNotCompressedClientDisabledNum;
+
+	/** The number of outgoing packets that were not compressed, due to the packet being smaller that the CVar 'net.OodleMinSizeForCompression'. */
+	uint64 OutNotCompressedTooSmallNum;
 
-	/** The number of outgoing packets that were not compressed, due to Oodle failing to compress - which were KeepAlive packets */
-	uint64 OutNotCompressedFailedKeepAliveNum;
 
 	/** The compressed length + decompression data overhead, of all outgoing packets. The most accurate measure of compression savings. */
 	uint64 OutCompressedWithOverheadLengthTotal;
@@ -66,8 +78,14 @@ public:
 	/** The compressed length of all outgoing packets. Measures Oodle algorithm compression, minus overhead reducing final savings. */
 	uint64 OutCompressedLengthTotal;
 
-	/** The length prior to compression, of all outgoing packets. */
+	/** The length prior to compression, of all outgoing packets (only counted for successfully compressed packets). */
 	uint64 OutBeforeCompressedLengthTotal;
+
+	/** The length of all outgoing packets, which failed to compress. */
+	uint64 OutNotCompressedFailedLengthTotal;
+
+	/** The length of all outgoing packets, where compression was skipped. */
+	uint64 OutNotCompressedSkippedLengthTotal;
 };
 
 /**
