Fixing errors

diff --git a/Engine/Plugins/Runtime/PacketHandlers/CompressionComponents/Oodle/Source/OodleCompressionFormat/Private/OodleCompressionFormat.cpp b/Engine/Plugins/Runtime/PacketHandlers/CompressionComponents/Oodle/Source/OodleCompressionFormat/Private/OodleCompressionFormat.cpp
--- a/Engine/Plugins/Runtime/PacketHandlers/CompressionComponents/Oodle/Source/OodleCompressionFormat/Private/OodleCompressionFormat.cpp
+++ b/Engine/Plugins/Runtime/PacketHandlers/CompressionComponents/Oodle/Source/OodleCompressionFormat/Private/OodleCompressionFormat.cpp
@@ -3,6 +3,7 @@
 #include "CoreTypes.h"
 #include "Misc/Compression.h"
 #include "Misc/ICompressionFormat.h"
+#include "Misc/CommandLine.h"
 #include "oodle2.h"
 
 DEFINE_LOG_CATEGORY_STATIC(OodleCompression, Log, All);
