Fixed compilation issue caused by missing Ooodle headers when grabbing code from Github

diff --git a/Engine/Plugins/Runtime/PacketHandlers/CompressionComponents/Oodle/Source/OodleCompressionFormat/Private/OodleCompressionFormat.cpp b/Engine/Plugins/Runtime/PacketHandlers/CompressionComponents/Oodle/Source/OodleCompressionFormat/Private/OodleCompressionFormat.cpp
--- a/Engine/Plugins/Runtime/PacketHandlers/CompressionComponents/Oodle/Source/OodleCompressionFormat/Private/OodleCompressionFormat.cpp
+++ b/Engine/Plugins/Runtime/PacketHandlers/CompressionComponents/Oodle/Source/OodleCompressionFormat/Private/OodleCompressionFormat.cpp
@@ -1,11 +1,14 @@
-// Copyright 1998-2019 Epic Games, Inc. All Rights Reserved.
+// Copyright 1998-2018 Epic Games, Inc. All Rights Reserved.
 
 #include "CoreTypes.h"
 #include "Misc/Compression.h"
 #include "Misc/ICompressionFormat.h"
 #include "Misc/CommandLine.h"
+#if HAS_OODLE_SDK
 #include "oodle2.h"
+#endif
 
+#if HAS_OODLE_SDK
 DEFINE_LOG_CATEGORY_STATIC(OodleCompression, Log, All);
 
 
@@ -86,8 +89,7 @@ struct FOodleCustomCompressor : ICompressionFormat
 	}
 };
 
-
-
+#endif
 
 #include "Misc/ICompressionFormat.h"
 
@@ -95,6 +97,7 @@ class FOodleCompressionFormatModuleInterface : public IModuleInterface
 {
 	virtual void StartupModule() override
 	{
+#if HAS_OODLE_SDK
 		FString Method = TEXT("Mermaid");
 		FString Level = TEXT("Normal");
 		int32 SpaceSpeedTradeoff = 256;
@@ -132,18 +135,18 @@ class FOodleCompressionFormatModuleInterface : public IModuleInterface
 		CompressionFormat = new FOodleCustomCompressor(MethodMap.FindRef(Method), LevelMap.FindRef(Level), SpaceSpeedTradeoff);
 
 		IModularFeatures::Get().RegisterModularFeature(COMPRESSION_FORMAT_FEATURE_NAME, CompressionFormat);
+#endif
 	}
 
 	virtual void ShutdownModule() override
 	{
+#if HAS_OODLE_SDK
 		IModularFeatures::Get().UnregisterModularFeature(COMPRESSION_FORMAT_FEATURE_NAME, CompressionFormat);
 		delete CompressionFormat;
+#endif
 	}
 
-	ICompressionFormat* CompressionFormat;
+	ICompressionFormat* CompressionFormat = nullptr;
 };
 
-IMPLEMENT_MODULE(FOodleCompressionFormatModuleInterface, OodleCompressionFormat);
-
-
-
+IMPLEMENT_MODULE(FOodleCompressionFormatModuleInterface, OodleCompressionFormat);
\ No newline at end of file
