Added extra safeguards to the Oodle dictionary generator, to prevent and detect generation of bad dictionaries.

diff --git a/Engine/Plugins/Runtime/PacketHandlers/CompressionComponents/Oodle/Source/OodleHandlerComponent/Classes/OodleTrainerCommandlet.h b/Engine/Plugins/Runtime/PacketHandlers/CompressionComponents/Oodle/Source/OodleHandlerComponent/Classes/OodleTrainerCommandlet.h
--- a/Engine/Plugins/Runtime/PacketHandlers/CompressionComponents/Oodle/Source/OodleHandlerComponent/Classes/OodleTrainerCommandlet.h
+++ b/Engine/Plugins/Runtime/PacketHandlers/CompressionComponents/Oodle/Source/OodleHandlerComponent/Classes/OodleTrainerCommandlet.h
@@ -266,8 +266,9 @@ private:
 	 *
 	 * @param InputCaptureFiles		As per BeginGenerateDictionary
 	 * @param bMergeAll				As per BeginGenerateDictionary
+	 * @return						Whether or not the capture file packets were read successfully
 	 */
-	void ReadPackets(const TArray<FString>& InputCaptureFiles);
+	bool ReadPackets(const TArray<FString>& InputCaptureFiles);
 
 	/**
 	 * Processes loaded packet data through the Oodle dictionary generation API, and then writes/compresses the result to the final file
diff --git a/Engine/Plugins/Runtime/PacketHandlers/CompressionComponents/Oodle/Source/OodleHandlerComponent/Private/OodleTrainerCommandlet.cpp b/Engine/Plugins/Runtime/PacketHandlers/CompressionComponents/Oodle/Source/OodleHandlerComponent/Private/OodleTrainerCommandlet.cpp
--- a/Engine/Plugins/Runtime/PacketHandlers/CompressionComponents/Oodle/Source/OodleHandlerComponent/Private/OodleTrainerCommandlet.cpp
+++ b/Engine/Plugins/Runtime/PacketHandlers/CompressionComponents/Oodle/Source/OodleHandlerComponent/Private/OodleTrainerCommandlet.cpp
@@ -357,8 +357,6 @@ void UOodleTrainerCommandlet::GetCaptureFiles(FString FilenameFilter, int32 Chan
 
 			return true;
 		}
-
-
 	};
 
 	FLocalDirIterator DirIterator(FilenameFilter, ChangelistNumber, OutFiles);
@@ -514,7 +512,7 @@ bool UOodleTrainerCommandlet::GenerateDictionary(bool bIsInput, int32 Changelist
 	}
 	else
 	{
-		UE_LOG(OodleHandlerComponentLog, Log, TEXT("Could not find any packet captures to for %s."), *InputOrOutput);
+		UE_LOG(OodleHandlerComponentLog, Log, TEXT("Could not find any packet captures for %s."), *InputOrOutput);
 	}
 	return bSuccess;
 }
@@ -527,7 +525,7 @@ bool UOodleTrainerCommandlet::GetMergeMapFromList(const TArray<FString>& FileLis
 	{
 		FArchive* ReadArc = IFileManager::Get().CreateFileReader(*CurMergeFile);
 
-		if (ReadArc != NULL)
+		if (ReadArc != nullptr)
 		{
 			OutMergeMap.Add(ReadArc, CurMergeFile);
 		}
@@ -540,6 +538,13 @@ bool UOodleTrainerCommandlet::GetMergeMapFromList(const TArray<FString>& FileLis
 		}
 	}
 
+	if (bSuccess && OutMergeMap.Num() == 0)
+	{
+		UE_LOG(OodleHandlerComponentLog, Error, TEXT("Could not find any packet capture files!"));
+
+		bSuccess = false;
+	}
+
 	return bSuccess;
 }
 
@@ -583,11 +588,7 @@ bool FOodleDictionaryGenerator::BeginGenerateDictionary(FString InOutputDictiona
 
 	bSuccess = InitGenerator();
 	bSuccess = bSuccess && UOodleTrainerCommandlet::GetMergeMapFromList(InputCaptureFiles, MergeMap);
-	
-	if (bSuccess)
-	{
-		ReadPackets(InputCaptureFiles);
-	}
+	bSuccess = bSuccess && ReadPackets(InputCaptureFiles);
 
 	if (bDebugDump)
 	{
@@ -696,33 +697,64 @@ bool FOodleDictionaryGenerator::InitGenerator()
 
 	bSuccess = UOodleTrainerCommandlet::VerifyOutputFile(OutputDictionaryFile);
 
+	// Pre-delete the output file, if it already exists - so failed dictionary generation does not leave behind the old file
+	if (bSuccess && FPaths::FileExists(OutputDictionaryFile))
+	{
+		bSuccess = IFileManager::Get().Delete(*OutputDictionaryFile, true);
+
+		if (!bSuccess)
+		{
+			UE_LOG(OodleHandlerComponentLog, Error, TEXT("Failed to delete old dictionary file: %s"), *OutputDictionaryFile);
+		}
+	}
+
 	return bSuccess;
 }
 
-void FOodleDictionaryGenerator::ReadPackets(const TArray<FString>& InputCaptureFiles)
+bool FOodleDictionaryGenerator::ReadPackets(const TArray<FString>& InputCaptureFiles)
 {
 	// Now begin training
+	bool bSuccess = true;
 	TArray<FArchive*> UnboundArchives;
 	TArray<FPacketCaptureArchive> BoundArchives;
 
-	MergeMap.GenerateKeyArray(UnboundArchives);
+	bSuccess = MergeMap.Num() > 0;
 
-	for (FArchive* CurArc : UnboundArchives)
+	if (bSuccess)
 	{
-		new(BoundArchives) FPacketCaptureArchive(*CurArc);
+		MergeMap.GenerateKeyArray(UnboundArchives);
+
+		for (FArchive* CurArc : UnboundArchives)
+		{
+			new(BoundArchives) FPacketCaptureArchive(*CurArc);
+		}
+	}
+	else
+	{
+		UE_LOG(OodleHandlerComponentLog, Error, TEXT("No capture files to process!"));
 	}
 
 
 	// Serialize capture archive headers, and retrieve packet count
 	int32 PacketCount = 0;
 
-	for (FPacketCaptureArchive& CurArc : BoundArchives)
+	if (bSuccess)
 	{
-		CurArc.SerializeCaptureHeader();
+		for (FPacketCaptureArchive& CurArc : BoundArchives)
+		{
+			CurArc.SerializeCaptureHeader();
 
-		if (!CurArc.IsError())
+			if (!CurArc.IsError())
+			{
+				PacketCount += CurArc.GetPacketCount();
+			}
+		}
+
+		bSuccess = PacketCount > 0;
+
+		if (!bSuccess)
 		{
-			PacketCount += CurArc.GetPacketCount();
+			UE_LOG(OodleHandlerComponentLog, Error, TEXT("No packets gathered from any capture files!"));;
 		}
 	}
 
@@ -874,8 +906,20 @@ void FOodleDictionaryGenerator::ReadPackets(const TArray<FString>& InputCaptureF
 		}
 	}
 
+	if (bSuccess)
+	{
+		bSuccess = PacketIdx > 0;
+
+		if (!bSuccess)
+		{
+			UE_LOG(OodleHandlerComponentLog, Error, TEXT("No packets successfully processed from capture files!"));
+		}
+	}
+
 	delete[] ReadBuffer;
 	ReadBuffer = nullptr;
+
+	return bSuccess;
 }
 
 bool FOodleDictionaryGenerator::GenerateAndWriteDictionary()
@@ -887,14 +931,14 @@ bool FOodleDictionaryGenerator::GenerateAndWriteDictionary()
 
 	UE_LOG(OodleHandlerComponentLog, Log, TEXT("Beginning dictionary generation..."));
 
-	UE_LOG(OodleHandlerComponentLog, Log, TEXT("- DictionaryPackets: %i, Size: %i (~%iMB)"),
+	UE_LOG(OodleHandlerComponentLog, Log, TEXT("- DictionaryPackets: %i, Size: %u (~%uMB)"),
 			DictionaryPackets.Num(), DictionaryPacketBytes, (DictionaryPacketBytes / (1024 * 1024)));
 
-	UE_LOG(OodleHandlerComponentLog, Log, TEXT("- DictionaryTestPackets: %i, Size: %i (~%iMB), Overflowed: %s"),
+	UE_LOG(OodleHandlerComponentLog, Log, TEXT("- DictionaryTestPackets: %i, Size: %u (~%uMB), Overflowed: %s"),
 			DictionaryTestPackets.Num(), DictionaryTestPacketBytes, (DictionaryTestPacketBytes / (1024 * 1024)),
 			(bDictionaryTestOverflow ? TEXT("Yes") : TEXT("No")));
 
-	UE_LOG(OodleHandlerComponentLog, Log, TEXT("- TrainerPackets: %i, Size: %i (~%iMB)"),
+	UE_LOG(OodleHandlerComponentLog, Log, TEXT("- TrainerPackets: %i, Size: %u (~%uMB)"),
 			TrainerPackets.Num(), TrainerPacketBytes, (TrainerPacketBytes / (1024 * 1024)));
 
 	UE_LOG(OodleHandlerComponentLog, Log, TEXT("Dictionary generation may take a very long time (up to 7 mins for 1GB of packets)."));
@@ -954,6 +998,21 @@ bool FOodleDictionaryGenerator::GenerateAndWriteDictionary()
 																		CompactCompressorStateBytes);
 
 
+		// Important warning for unusually small dictionary files - if they compress down this much, something is usually wrong.
+		bool bDeleteFile = false;
+
+		if (bSuccess && OutputFile.TotalSize() < 65536)
+		{
+			FText Msg = LOCTEXT("BadOodleDictionaryGen",
+				"The generated dictionary is less than 64KB. This is unusually small, indicating a problem - do NOT use for production! Delete the file?");
+
+			EAppReturnType::Type Result = FMessageDialog::Open(EAppMsgType::YesNo, Msg);
+
+			bDeleteFile = (Result == EAppReturnType::Yes);
+			bSuccess = false;
+		}
+
+
 		if (bSuccess)
 		{
 			UE_LOG(OodleHandlerComponentLog, Log, TEXT("Successfully processed packet captures, and wrote dictionary file."));
@@ -967,6 +1026,11 @@ bool FOodleDictionaryGenerator::GenerateAndWriteDictionary()
 		OutArc->Close();
 
 		delete OutArc;
+
+		if (bDeleteFile)
+		{
+			IFileManager::Get().Delete(*OutputDictionaryFile);
+		}
 	}
 	else
 	{
@@ -1002,10 +1066,10 @@ bool FOodleDictionaryGenerator::GenerateAndWriteDictionary()
 
 		UE_LOG(OodleHandlerComponentLog, Log, TEXT("Compression test results:"));
 
-		UE_LOG(OodleHandlerComponentLog, Log, TEXT("- CompressionTestPackets: %i, Size: %i (~%iMB)"),
+		UE_LOG(OodleHandlerComponentLog, Log, TEXT("- CompressionTestPackets: %i, Size: %u (~%uMB)"),
 				CompressionTestPackets.Num(), CompressionTestPacketBytes, (CompressionTestPacketBytes / (1024 * 1024)));
 
-		UE_LOG(OodleHandlerComponentLog, Log, TEXT("- Uncompressed: %i (~%iMB), Compressed: %i (~%iMB)"), TotalUncompressed,
+		UE_LOG(OodleHandlerComponentLog, Log, TEXT("- Uncompressed: %llu (~%lluMB), Compressed: %llu (~%lluMB)"), TotalUncompressed,
 				(TotalUncompressed / (1024 * 1024)), TotalCompressed, (TotalCompressed / (1024 * 1024)));
 
 		UE_LOG(OodleHandlerComponentLog, Log, TEXT("- Total Savings: %f%"),
