Fixed some Oodle analytics errors.

diff --git a/Engine/Plugins/Runtime/PacketHandlers/CompressionComponents/Oodle/Source/OodleHandlerComponent/Private/OodleHandlerComponent.cpp b/Engine/Plugins/Runtime/PacketHandlers/CompressionComponents/Oodle/Source/OodleHandlerComponent/Private/OodleHandlerComponent.cpp
--- a/Engine/Plugins/Runtime/PacketHandlers/CompressionComponents/Oodle/Source/OodleHandlerComponent/Private/OodleHandlerComponent.cpp
+++ b/Engine/Plugins/Runtime/PacketHandlers/CompressionComponents/Oodle/Source/OodleHandlerComponent/Private/OodleHandlerComponent.cpp
@@ -932,8 +932,12 @@ void OodleHandlerComponent::Incoming(FBitReader& Packet)
 			if (bOodleAnalytics && NetAnalyticsData.IsValid())
 			{
 				FOodleAnalyticsVars* AnalyticsVars = NetAnalyticsData->GetLocalData();
+				uint32 PacketSize = Packet.GetBytesLeft();
 
 				AnalyticsVars->InNotCompressedNum++;
+				AnalyticsVars->InCompressedWithOverheadLengthTotal += PacketSize;
+				AnalyticsVars->InCompressedLengthTotal += PacketSize;
+				AnalyticsVars->InDecompressedLengthTotal += PacketSize;
 			}
 
 #if !UE_BUILD_SHIPPING || OODLE_DEV_SHIPPING
@@ -1055,6 +1059,12 @@ void OodleHandlerComponent::Outgoing(FBitWriter& Packet, FOutPacketTraits& Trait
 					}
 					else
 					{
+						if (AnalyticsVars != nullptr)
+						{
+							AnalyticsVars->OutCompressedLengthTotal += UncompressedBytes;
+							AnalyticsVars->OutCompressedWithOverheadLengthTotal += UncompressedBytes;
+						}
+
 						// @todo #JohnB: Try to eliminate this appBitscpy, by reserving the bCompressedPacket bit and other data,
 						//					at the start of the bit writer
 						Packet.SerializeBits(UncompressedData, UncompressedBits);
