- Updated Oodle SDK location to Engine/ThirdParty
- Updated Oodle to 280
- Removed older versions
- Cleaned up build files and related .cpp code, greatly decreases the complexity of using the plugin and the dlls on windows, etc.
- Will be much easier for RAD to take ownership of the plugin and distribute to their licensees

diff --git a/Engine/Plugins/Runtime/PacketHandlers/CompressionComponents/Oodle/Source/OodleCompressionFormat/OodleCompressionFormat.Build.cs b/Engine/Plugins/Runtime/PacketHandlers/CompressionComponents/Oodle/Source/OodleCompressionFormat/OodleCompressionFormat.Build.cs
--- a/Engine/Plugins/Runtime/PacketHandlers/CompressionComponents/Oodle/Source/OodleCompressionFormat/OodleCompressionFormat.Build.cs
+++ b/Engine/Plugins/Runtime/PacketHandlers/CompressionComponents/Oodle/Source/OodleCompressionFormat/OodleCompressionFormat.Build.cs
@@ -9,14 +9,8 @@ public class OodleCompressionFormat : ModuleRules
 	{
 		// @todo oodle: Clean this up with the handler component?
 
-		
-		// this needs to match the version in Oodle.Build.cs
-		string OodleVersion = "255";
-
 		ShortName = "OodleFormat";
-
-		BinariesSubFolder = "NotForLicensees";
-
+		
 		PublicDependencyModuleNames.AddRange(
 			new string[]
 			{
@@ -26,48 +20,39 @@ public class OodleCompressionFormat : ModuleRules
 				"Analytics"
 			});
 
-		string PlatformName = Target.Platform.ToString();
 		if (Target.Platform == UnrealTargetPlatform.Win64 || Target.Platform == UnrealTargetPlatform.Win32)
 		{
-			PlatformName = "win";
-
 			// this is needed to hunt down the DLL in the binaries directory for running unstaged
 			PrivateDependencyModuleNames.AddRange(
 			new string[] {
 				"Projects",
 			});
 		}
-		else if (Target.IsInPlatformGroup(UnrealPlatformGroup.Unix))
-		{
-			PlatformName = "Linux";
-		}
 
 		// Check the NotForLicensees folder first
-		string OodleNotForLicenseesLibDir = System.IO.Path.Combine(Target.UEThirdPartySourceDirectory, "..", "..",
-			"Plugins", "Runtime", "PacketHandlers", "CompressionComponents", "Oodle", "Source", "ThirdParty", "NotForLicensees",
-			"Oodle", OodleVersion, PlatformName, "lib");
+		string OodleNotForLicenseesBuildFile = System.IO.Path.Combine(Target.UEThirdPartySourceDirectory, 
+			"NotForLicensees/RadGames/OodleData/OodleData.Build.cs");
+		string OodleLicenseesBuildFile = System.IO.Path.Combine(Target.UEThirdPartySourceDirectory,
+			"RadGames/OodleData/OodleData.Build.cs");
+
+		bool bHaveNFLOodleSDK = System.IO.File.Exists(OodleNotForLicenseesBuildFile);
+		bool bHaveLicenseesOodleSDK = System.IO.File.Exists(OodleLicenseesBuildFile);
 
-		bool bHaveOodleSDK = false;
-		if (OodleNotForLicenseesLibDir.Length > 0)
+		// if we are loading from NFL, save binary in NFL
+		if (bHaveNFLOodleSDK)
 		{
-			try
-			{
-				bHaveOodleSDK = System.IO.Directory.Exists(OodleNotForLicenseesLibDir);
-			}
-			catch (System.Exception)
-			{
-			}
+			BinariesSubFolder = "NotForLicensees";
 		}
 
-		if (bHaveOodleSDK)
+		if (bHaveNFLOodleSDK || bHaveLicenseesOodleSDK)
 		{
-			AddEngineThirdPartyPrivateStaticDependencies(Target, "Oodle");
-			PublicIncludePathModuleNames.Add("Oodle");
-			PublicDefinitions.Add("HAS_OODLE_SDK=1");
+			AddEngineThirdPartyPrivateStaticDependencies(Target, "OodleData");
+			PublicIncludePathModuleNames.Add("OodleData");
+			PublicDefinitions.Add("HAS_OODLE_DATA_SDK=1");
 		}
 		else
 		{
-			PublicDefinitions.Add("HAS_OODLE_SDK=0");
+			PublicDefinitions.Add("HAS_OODLE_DATA_SDK=0");
 		}
 	}
 }
\ No newline at end of file
diff --git a/Engine/Plugins/Runtime/PacketHandlers/CompressionComponents/Oodle/Source/OodleCompressionFormat/Private/OodleCompressionFormat.cpp b/Engine/Plugins/Runtime/PacketHandlers/CompressionComponents/Oodle/Source/OodleCompressionFormat/Private/OodleCompressionFormat.cpp
--- a/Engine/Plugins/Runtime/PacketHandlers/CompressionComponents/Oodle/Source/OodleCompressionFormat/Private/OodleCompressionFormat.cpp
+++ b/Engine/Plugins/Runtime/PacketHandlers/CompressionComponents/Oodle/Source/OodleCompressionFormat/Private/OodleCompressionFormat.cpp
@@ -4,14 +4,12 @@
 #include "Misc/Compression.h"
 #include "Misc/ICompressionFormat.h"
 #include "Misc/CommandLine.h"
-#if HAS_OODLE_SDK
-#include "oodle2.h"
-#endif
 
-#if HAS_OODLE_SDK
-DEFINE_LOG_CATEGORY_STATIC(OodleCompression, Log, All);
+#if HAS_OODLE_DATA_SDK
 
+#include "oodle2.h"
 
+DEFINE_LOG_CATEGORY_STATIC(OodleCompression, Log, All);
 
 struct FOodleCustomCompressor : ICompressionFormat
 {
@@ -84,7 +82,7 @@ struct FOodleCustomCompressor : ICompressionFormat
 	{
 		ConditionalInitialize();
 
-		int32 Needed = (int32)OodleLZ_GetCompressedBufferSizeNeeded(UncompressedSize);
+		int32 Needed = (int32)OodleLZ_GetCompressedBufferSizeNeeded(Compressor, UncompressedSize);
 		return Needed;
 	}
 };
@@ -97,7 +95,8 @@ class FOodleCompressionFormatModuleInterface : public IModuleInterface
 {
 	virtual void StartupModule() override
 	{
-#if HAS_OODLE_SDK
+#if HAS_OODLE_DATA_SDK
+
 		FString Method = TEXT("Mermaid");
 		FString Level = TEXT("Normal");
 		int32 SpaceSpeedTradeoff = 256;
@@ -112,13 +111,12 @@ class FOodleCompressionFormatModuleInterface : public IModuleInterface
 			{ TEXT("Mermaid"), OodleLZ_Compressor_Mermaid },
 			{ TEXT("Kraken"), OodleLZ_Compressor_Kraken },
 			{ TEXT("Selkie"), OodleLZ_Compressor_Selkie },
-			{ TEXT("LZNA"), OodleLZ_Compressor_LZNA },
-			{ TEXT("BitKnit"), OodleLZ_Compressor_BitKnit },
-			{ TEXT("LZB16"), OodleLZ_Compressor_LZB16 },
+#if UE4_OODLE_VER >= 270
+			{ TEXT("Leviathan"), OodleLZ_Compressor_Leviathan },
+#endif
 		};
 		TMap<FString, OodleLZ_CompressionLevel> LevelMap = { 
 			{ TEXT("None"), OodleLZ_CompressionLevel_None },
-			{ TEXT("RLE"), OodleLZ_CompressionLevel_RLE },
 			{ TEXT("VeryFast"), OodleLZ_CompressionLevel_VeryFast },
 			{ TEXT("Fast"), OodleLZ_CompressionLevel_Fast },
 			{ TEXT("Normal"), OodleLZ_CompressionLevel_Normal },
@@ -135,12 +133,15 @@ class FOodleCompressionFormatModuleInterface : public IModuleInterface
 		CompressionFormat = new FOodleCustomCompressor(MethodMap.FindRef(Method), LevelMap.FindRef(Level), SpaceSpeedTradeoff);
 
 		IModularFeatures::Get().RegisterModularFeature(COMPRESSION_FORMAT_FEATURE_NAME, CompressionFormat);
+
+		OodleCore_Plugins_SetPrintf(OodleCore_Plugin_Printf_Verbose);
+
 #endif
 	}
 
 	virtual void ShutdownModule() override
 	{
-#if HAS_OODLE_SDK
+#if HAS_OODLE_DATA_SDK
 		IModularFeatures::Get().UnregisterModularFeature(COMPRESSION_FORMAT_FEATURE_NAME, CompressionFormat);
 		delete CompressionFormat;
 #endif
@@ -149,4 +150,4 @@ class FOodleCompressionFormatModuleInterface : public IModuleInterface
 	ICompressionFormat* CompressionFormat = nullptr;
 };
 
-IMPLEMENT_MODULE(FOodleCompressionFormatModuleInterface, OodleCompressionFormat);
\ No newline at end of file
+IMPLEMENT_MODULE(FOodleCompressionFormatModuleInterface, OodleCompressionFormat);
diff --git a/Engine/Plugins/Runtime/PacketHandlers/CompressionComponents/Oodle/Source/OodleHandlerComponent/Classes/OodleTrainerCommandlet.h b/Engine/Plugins/Runtime/PacketHandlers/CompressionComponents/Oodle/Source/OodleHandlerComponent/Classes/OodleTrainerCommandlet.h
--- a/Engine/Plugins/Runtime/PacketHandlers/CompressionComponents/Oodle/Source/OodleHandlerComponent/Classes/OodleTrainerCommandlet.h
+++ b/Engine/Plugins/Runtime/PacketHandlers/CompressionComponents/Oodle/Source/OodleHandlerComponent/Classes/OodleTrainerCommandlet.h
@@ -10,7 +10,7 @@
 #include "OodleTrainerCommandlet.generated.h"
 
 // don't compile this in for non-editor
-#define USE_OODLE_TRAINER_COMMANDLET (HAS_OODLE_SDK && WITH_EDITOR)
+#define USE_OODLE_TRAINER_COMMANDLET (HAS_OODLE_NET_SDK && WITH_EDITOR)
 
 /**
  * Commandlet for processing UE4 packet captures, through Oodle's training API, for generating compressed state dictionaries.
@@ -156,7 +156,7 @@ private:
 	 */
 	bool GenerateDictionary(bool bIsInput, int32 ChangelistNumber);
 
-#endif // HAS_OODLE_SDK
+#endif // HAS_OODLE_NET_SDK
 #endif // !UE_BUILD_SHIPPING || OODLE_DEV_SHIPPING
 };
 
diff --git a/Engine/Plugins/Runtime/PacketHandlers/CompressionComponents/Oodle/Source/OodleHandlerComponent/OodleHandlerComponent.Build.cs b/Engine/Plugins/Runtime/PacketHandlers/CompressionComponents/Oodle/Source/OodleHandlerComponent/OodleHandlerComponent.Build.cs
--- a/Engine/Plugins/Runtime/PacketHandlers/CompressionComponents/Oodle/Source/OodleHandlerComponent/OodleHandlerComponent.Build.cs
+++ b/Engine/Plugins/Runtime/PacketHandlers/CompressionComponents/Oodle/Source/OodleHandlerComponent/OodleHandlerComponent.Build.cs
@@ -7,16 +7,8 @@ public class OodleHandlerComponent : ModuleRules
 {
     public OodleHandlerComponent(ReadOnlyTargetRules Target) : base(Target)
     {
-		// @todo oodle: Clean this up with the compression format?
-
-
-		// this needs to match the version in Oodle.Build.cs
-		string OodleVersion = "255";
-
 		ShortName = "OodleHC";
 
-        BinariesSubFolder = "NotForLicensees";
-		
 		PrivateIncludePaths.Add("OodleHandlerComponent/Private");
 
 		PublicDependencyModuleNames.AddRange(
@@ -29,52 +21,60 @@ public class OodleHandlerComponent : ModuleRules
                 "Analytics"
 			});
 
-		PrivateDependencyModuleNames.AddRange(
-			new string[] {
-				"Projects",
-			});
 
-		string PlatformName = Target.Platform.ToString();
 		if (Target.Platform == UnrealTargetPlatform.Win64 || Target.Platform == UnrealTargetPlatform.Win32)
 		{
-			PlatformName = "win";
-
 			// this is needed to hunt down the DLL in the binaries directory for running unstaged
 			PrivateDependencyModuleNames.AddRange(
 			new string[] {
 				"Projects",
 			});
 		}
-		else if (Target.IsInPlatformGroup(UnrealPlatformGroup.Unix))
+
+		string OodleNotForLicenseesNetBuildFile = System.IO.Path.Combine(Target.UEThirdPartySourceDirectory,
+			"NotForLicensees/RadGames/OodleNet/OodleNet.Build.cs");
+		string OodleLicenseesNetBuildFile = System.IO.Path.Combine(Target.UEThirdPartySourceDirectory,
+			"RadGames/OodleNet/OodleNet.Build.cs");
+
+		string OodleNotForLicenseesDataBuildFile = System.IO.Path.Combine(Target.UEThirdPartySourceDirectory,
+			"NotForLicensees/RadGames/OodleData/OodleData.Build.cs");
+		string OodleLicenseesDataBuildFile = System.IO.Path.Combine(Target.UEThirdPartySourceDirectory,
+			"RadGames/OodleData/OodleData.Build.cs");
+
+		bool bHaveNFLOodleNetSDK = System.IO.File.Exists(OodleNotForLicenseesNetBuildFile);
+		bool bHaveLicenseesOodleNetSDK = System.IO.File.Exists(OodleLicenseesNetBuildFile);
+		bool bHaveNFLOodleDataSDK = System.IO.File.Exists(OodleNotForLicenseesDataBuildFile);
+		bool bHaveLicenseesOodleDataSDK = System.IO.File.Exists(OodleLicenseesDataBuildFile);
+
+		// if we are loading from NFL, save binary in NFL
+		if (bHaveNFLOodleDataSDK || bHaveNFLOodleNetSDK)
 		{
-			PlatformName = "Linux";
+			BinariesSubFolder = "NotForLicensees";
 		}
 
-		// Check the NotForLicensees folder first
-		string OodleNotForLicenseesLibDir = System.IO.Path.Combine(ModuleDirectory, "..", "ThirdParty", "NotForLicensees",
-			"Oodle", OodleVersion, PlatformName, "lib");
-
-		bool bHaveOodleSDK = false;
-		if (OodleNotForLicenseesLibDir.Length > 0)
+		if (bHaveNFLOodleNetSDK || bHaveLicenseesOodleNetSDK)
 		{
-			try
+	        AddEngineThirdPartyPrivateStaticDependencies(Target, "OodleNet");
+	        PublicIncludePathModuleNames.Add("OodleNet");
+			PublicDefinitions.Add( "HAS_OODLE_NET_SDK=1" );
+
+			// data SDK is optional
+			if (bHaveNFLOodleDataSDK || bHaveLicenseesOodleDataSDK)
 			{
-				bHaveOodleSDK = System.IO.Directory.Exists( OodleNotForLicenseesLibDir );
+				AddEngineThirdPartyPrivateStaticDependencies(Target, "OodleData");
+				PublicIncludePathModuleNames.Add("OodleData");
+				PublicDefinitions.Add("HAS_OODLE_DATA_SDK=1");
 			}
-			catch ( System.Exception )
+			else
 			{
+				PublicDefinitions.Add("HAS_OODLE_DATA_SDK=0");
 			}
-        }
 
-		if ( bHaveOodleSDK )
-		{
-	        AddEngineThirdPartyPrivateStaticDependencies(Target, "Oodle");
-	        PublicIncludePathModuleNames.Add("Oodle");
-			PublicDefinitions.Add( "HAS_OODLE_SDK=1" );
 		}
 		else
 		{
-			PublicDefinitions.Add( "HAS_OODLE_SDK=0" );
+			PublicDefinitions.Add("HAS_OODLE_NET_SDK=0");
+			PublicDefinitions.Add("HAS_OODLE_DATA_SDK=0");
 		}
-    }
+	}
 }
\ No newline at end of file
diff --git a/Engine/Plugins/Runtime/PacketHandlers/CompressionComponents/Oodle/Source/OodleHandlerComponent/Private/OodleArchives.cpp b/Engine/Plugins/Runtime/PacketHandlers/CompressionComponents/Oodle/Source/OodleHandlerComponent/Private/OodleArchives.cpp
--- a/Engine/Plugins/Runtime/PacketHandlers/CompressionComponents/Oodle/Source/OodleHandlerComponent/Private/OodleArchives.cpp
+++ b/Engine/Plugins/Runtime/PacketHandlers/CompressionComponents/Oodle/Source/OodleHandlerComponent/Private/OodleArchives.cpp
@@ -3,7 +3,7 @@
 #include "OodleArchives.h"
 #include "OodleHandlerComponent.h"
 
-#if HAS_OODLE_SDK
+#if HAS_OODLE_NET_SDK
 
 // Maximum size of compress/decompress buffers (just under 2GB, due to max int32 value)
 #define MAX_COMPRESS_BUFFER (1024 * 1024 * 2047)
@@ -42,10 +42,11 @@ bool FOodleArchiveBase::SerializeOodleCompressData(FOodleCompressedData& OutData
 	{
 		OutDataInfo.DecompressedLength.Set(*this, DataBytes);
 	
-		uint32 CompressBufferLen = OodleLZ_GetCompressedBufferSizeNeeded(DataBytes);
+#if HAS_OODLE_DATA_SDK
+		uint32 CompressBufferLen = OodleLZ_GetCompressedBufferSizeNeeded(OodleLZ_Compressor_Kraken, DataBytes);
 		uint8* CompressBuffer = new uint8[CompressBufferLen];
 
-		SINTa OodleLen = OodleLZ_Compress(OodleLZ_Compressor_LZNIB, Data, (SINTa)DataBytes, CompressBuffer,
+		SINTa OodleLen = OodleLZ_Compress(OodleLZ_Compressor_Kraken, Data, (SINTa)DataBytes, CompressBuffer,
 											OodleLZ_CompressionLevel_Optimal);
 
 		bSuccess = ensure(OodleLen != OODLELZ_FAILED);
@@ -59,9 +60,22 @@ bool FOodleArchiveBase::SerializeOodleCompressData(FOodleCompressedData& OutData
 
 			InnerArchive.Serialize((void*)CompressBuffer, OodleLen);
 		}
+		else
+#endif
+		{
+			// no compression
+			OutDataInfo.CompressedLength.Set(*this, DataBytes);
+
+			uint32 OffsetPos = InnerArchive.Tell();
+			OutDataInfo.Offset.Set(*this, OffsetPos);
+
+			InnerArchive.Serialize((void*)Data, DataBytes);
+		}
 
+#if HAS_OODLE_DATA_SDK
+		delete [] CompressBuffer;
+#endif
 
-		delete[] CompressBuffer;
 	}
 
 	return bSuccess;
@@ -85,35 +99,47 @@ bool FOodleArchiveBase::SerializeOodleDecompressData(FOodleCompressedData& DataI
 	{
 		SeekPush(DataOffset);
 
-		// @todo #JohnB: Remove bOutDataSlack after Oodle update, and after checking with Luigi
-		uint8* CompressedData = new uint8[CompressedLength + (bOutDataSlack ? OODLE_DICTIONARY_SLACK : 0)];
-		uint8* DecompressedData = new uint8[DecompressedLength];
-
-		// @todo #JohnB: Remove after Oodle update, and after checking with Luigi
-		if (bOutDataSlack)
+		if (CompressedLength == DecompressedLength)   //wasn't actually compressed
 		{
-			FMemory::Memzero(CompressedData + CompressedLength, OODLE_DICTIONARY_SLACK);
+			uint8 * RawData = new uint8[DecompressedLength];
+			InnerArchive.Serialize(RawData, DecompressedLength);
+			OutData = RawData;
+			OutDataBytes = DecompressedLength;
 		}
+		else
+		{
+#if HAS_OODLE_DATA_SDK
+			// @todo #JohnB: Remove bOutDataSlack after Oodle update, and after checking with Luigi
+			uint8* CompressedData = new uint8[CompressedLength + (bOutDataSlack ? OODLE_DICTIONARY_SLACK : 0)];
+			uint8* DecompressedData = new uint8[DecompressedLength];
 
-		InnerArchive.Serialize(CompressedData, CompressedLength);
+			// @todo #JohnB: Remove after Oodle update, and after checking with Luigi
+			if (bOutDataSlack)
+			{
+				FMemory::Memzero(CompressedData + CompressedLength, OODLE_DICTIONARY_SLACK);
+			}
 
-		SINTa OodleLen = OodleLZ_Decompress((void*)CompressedData, CompressedLength, (void*)DecompressedData, DecompressedLength);
+			InnerArchive.Serialize(CompressedData, CompressedLength);
 
-		bSuccess = OodleLen != 0;
+			SINTa OodleLen = OodleLZ_Decompress((void*)CompressedData, CompressedLength, (void*)DecompressedData, DecompressedLength, OodleLZ_FuzzSafe_Yes);
 
-		if (bSuccess)
-		{
-			check(OodleLen == DecompressedLength);
+			bSuccess = OodleLen != 0;
 
-			OutData = DecompressedData;
-			OutDataBytes = DecompressedLength;
-		}
-		else
-		{
-			delete[] DecompressedData;
-		}
+			if (bSuccess)
+			{
+				check(OodleLen == DecompressedLength);
 
-		delete[] CompressedData;
+				OutData = DecompressedData;
+				OutDataBytes = DecompressedLength;
+			}
+			else
+			{
+				delete[] DecompressedData;
+			}
+
+			delete[] CompressedData;
+#endif
+		}
 
 		SeekPop();
 	}
@@ -349,7 +375,7 @@ void FOodleDictionaryArchive::FDictionaryHeader::SerializeHeader(FOodleDictionar
 
 		bSuccess = bSuccess && ensure(Magic == DICTIONARY_HEADER_MAGIC);
 		bSuccess = bSuccess && ensure(DictionaryVersion <= DICTIONARY_FILE_VERSION);
-		bSuccess = bSuccess && ensure(OodleMajorHeaderVersion == OODLE2_VERSION_MAJOR);
+		//bSuccess = bSuccess && ensure(OodleMajorHeaderVersion == OODLE2NET_VERSION_MAJOR);   //deprecated. oodle network data format has stabilised.
 	}
 
 	if (!bSuccess)
@@ -370,7 +396,7 @@ FOodleDictionaryArchive::FOodleDictionaryArchive(FArchive& InInnerArchive)
 	if (IsSaving())
 	{
 		Header.DictionaryVersion = DICTIONARY_FILE_VERSION;
-		Header.OodleMajorHeaderVersion = OODLE2_VERSION_MAJOR;
+		Header.OodleMajorHeaderVersion = OODLE2NET_VERSION_MAJOR;
 	}
 }
 
diff --git a/Engine/Plugins/Runtime/PacketHandlers/CompressionComponents/Oodle/Source/OodleHandlerComponent/Private/OodleHandlerComponent.cpp b/Engine/Plugins/Runtime/PacketHandlers/CompressionComponents/Oodle/Source/OodleHandlerComponent/Private/OodleHandlerComponent.cpp
--- a/Engine/Plugins/Runtime/PacketHandlers/CompressionComponents/Oodle/Source/OodleHandlerComponent/Private/OodleHandlerComponent.cpp
+++ b/Engine/Plugins/Runtime/PacketHandlers/CompressionComponents/Oodle/Source/OodleHandlerComponent/Private/OodleHandlerComponent.cpp
@@ -33,7 +33,8 @@ DEFINE_LOG_CATEGORY(OodleHandlerComponentLog);
 // @todo #JohnB: If you could unwrap the analytics data shared pointer, by guaranteeing its lifetime, that would be a good optimization,
 //					as the multithreaded version is presently a bit expensive
 
-#if HAS_OODLE_SDK
+#if HAS_OODLE_NET_SDK
+
 #define OODLE_INI_SECTION TEXT("OodleHandlerComponent")
 
 // @todo #JohnB: Remove after Oodle update, and after checking with Luigi
@@ -1065,7 +1066,7 @@ void OodleHandlerComponent::Outgoing(FBitWriter& Packet, FOutPacketTraits& Trait
 			uint32 UncompressedBits = Packet.GetNumBits();
 
 			bool bWithinBitBounds = UncompressedBits > 0 && ensure(UncompressedBits <= MaxAdjustedLengthBits) &&
-									ensure(OodleLZ_GetCompressedBufferSizeNeeded(UncompressedBytes) <= MAX_OODLE_BUFFER);
+				ensure(OodleNetwork1_CompressedBufferSizeNeeded(UncompressedBytes) <= MAX_OODLE_BUFFER);
 
 			if (bWithinBitBounds)
 			{
@@ -1565,27 +1566,20 @@ void FOodleComponentModuleInterface::StartupModule()
 	GOodleSaveDir = FPaths::ConvertRelativePathToFull(FPaths::Combine(*FPaths::ProjectSavedDir(), TEXT("Oodle")));
 	GOodleContentDir = FPaths::ConvertRelativePathToFull(FPaths::Combine(*FPaths::ProjectContentDir(), TEXT("Oodle")));
 
-	// @todo #JohnB: Remove after Oodle update, and after checking with Luigi
-	OodlePlugins_SetAssertion(&UEOodleDisplayAssert);
+// 	// @todo #JohnB: Remove after Oodle update, and after checking with Luigi
+// 	OodlePlugins_SetAssertion(&UEOodleDisplayAssert);
 }
 
 void FOodleComponentModuleInterface::ShutdownModule()
 {
-	// @todo #JohnB: Remove after Oodle update, and after checking with Luigi
-	OodlePlugins_SetAssertion(nullptr);
-
-#if PLATFORM_WINDOWS
-	if (OodleDllHandle != nullptr)
-	{
-		FPlatformProcess::FreeDllHandle(OodleDllHandle);
-		OodleDllHandle = nullptr;
-	}
-#endif
+// 	// @todo #JohnB: Remove after Oodle update, and after checking with Luigi
+// 	OodlePlugins_SetAssertion(nullptr);
 }
+
 #else
 TSharedPtr<HandlerComponent> FOodleComponentModuleInterface::CreateComponentInstance(FString& Options)
 {
-	UE_LOG(OodleHandlerComponentLog, Error, TEXT("Can't create OodleHandlerComponent instance - HAS_OODLE_SDK is false."));
+	UE_LOG(OodleHandlerComponentLog, Error, TEXT("Can't create OodleHandlerComponent instance - HAS_OODLE_NET_SDK is false."));
 
 	return TSharedPtr<HandlerComponent>(nullptr);
 }
diff --git a/Engine/Plugins/Runtime/PacketHandlers/CompressionComponents/Oodle/Source/OodleHandlerComponent/Private/OodleTrainerCommandlet.cpp b/Engine/Plugins/Runtime/PacketHandlers/CompressionComponents/Oodle/Source/OodleHandlerComponent/Private/OodleTrainerCommandlet.cpp
--- a/Engine/Plugins/Runtime/PacketHandlers/CompressionComponents/Oodle/Source/OodleHandlerComponent/Private/OodleTrainerCommandlet.cpp
+++ b/Engine/Plugins/Runtime/PacketHandlers/CompressionComponents/Oodle/Source/OodleHandlerComponent/Private/OodleTrainerCommandlet.cpp
@@ -1051,8 +1051,6 @@ bool FOodleDictionaryGenerator::GenerateAndWriteDictionary()
 			uint8* CurPacket = CompressionTestPackets[CurPacketIdx];
 			int32 CurPacketSize = CompressionTestPacketSizes[CurPacketIdx];
 
-			check(OodleLZ_GetCompressedBufferSizeNeeded(CurPacketSize) <= MAX_OODLE_BUFFER);
-
 			SINTa CompressedLengthSINT = 0;
 			
 			CompressedLengthSINT = OodleNetwork1UDP_Encode(CompressorState, SharedDictionary, CurPacket, CurPacketSize,
@@ -1209,7 +1207,7 @@ void FOodleDictionaryGenerator::Cleanup()
 		MergeMap.Empty();
 	}
 }
-#endif // HAS_OODLE_SDK
+#endif
 
 #endif // !UE_BUILD_SHIPPING || OODLE_DEV_SHIPPING
 
diff --git a/Engine/Plugins/Runtime/PacketHandlers/CompressionComponents/Oodle/Source/OodleHandlerComponent/Private/OodleUtils.cpp b/Engine/Plugins/Runtime/PacketHandlers/CompressionComponents/Oodle/Source/OodleHandlerComponent/Private/OodleUtils.cpp
--- a/Engine/Plugins/Runtime/PacketHandlers/CompressionComponents/Oodle/Source/OodleHandlerComponent/Private/OodleUtils.cpp
+++ b/Engine/Plugins/Runtime/PacketHandlers/CompressionComponents/Oodle/Source/OodleHandlerComponent/Private/OodleUtils.cpp
@@ -3,11 +3,12 @@
 #include "OodleUtils.h"
 #include "Serialization/MemoryReader.h"
 #include "Serialization/MemoryWriter.h"
-#if HAS_OODLE_SDK
+
+#if HAS_OODLE_DATA_SDK // allow data SDK usage if we have it
 #include "oodle2.h"
 #endif
 
-#if HAS_OODLE_SDK
+#if HAS_OODLE_NET_SDK
 
 namespace OodleUtils
 {
@@ -22,7 +23,12 @@ namespace OodleUtils
 
 		OutBuffer.SetNum(Size, false);
 
-		return (OodleLZ_Decompress(InCompressed.GetData() + Reader.Tell(), (SINTa)CompressedSize, OutBuffer.GetData(), OutBuffer.Num()) == OutBuffer.Num());
+#if HAS_OODLE_DATA_SDK
+		return (OodleLZ_Decompress(InCompressed.GetData() + Reader.Tell(), (SINTa)CompressedSize, OutBuffer.GetData(), OutBuffer.Num(), OodleLZ_FuzzSafe_Yes) == OutBuffer.Num());
+#else
+		bool bSuccess = FCompression::UncompressMemory(NAME_Zlib, OutBuffer.GetData(), OutBuffer.Num(), (void*)(InCompressed.GetData() + Reader.Tell()), CompressedSize);
+		return bSuccess;
+#endif
 	}
 
 	bool CompressReplayData(const TArray<uint8>& InBuffer, TArray< uint8 >& OutCompressed)
@@ -36,17 +42,27 @@ namespace OodleUtils
 
 		int64 ReservedBytes = Writer.Tell();
 
-		OutCompressed.SetNum(ReservedBytes + OodleLZ_GetCompressedBufferSizeNeeded(InBuffer.Num()));
+#if HAS_OODLE_DATA_SDK
 
-		const OodleLZ_Compressor			Compressor = OodleLZ_Compressor_LZB16;
+		const OodleLZ_Compressor			Compressor = OodleLZ_Compressor_Selkie;
 		const OodleLZ_CompressionLevel		Level = OodleLZ_CompressionLevel_VeryFast;
 
+		OutCompressed.SetNum(ReservedBytes + OodleLZ_GetCompressedBufferSizeNeeded(Compressor, InBuffer.Num()));
 		CompressedSize = OodleLZ_Compress(Compressor, InBuffer.GetData(), (SINTa)InBuffer.Num(), OutCompressed.GetData() + ReservedBytes, Level);
 
 		if (CompressedSize == OODLELZ_FAILED)
 		{
 			return false;
 		}
+#else
+		OutCompressed.SetNum(ReservedBytes + FCompression::CompressMemoryBound(NAME_Zlib, InBuffer.Num()));
+		bool bSuccess = FCompression::CompressMemory(NAME_Zlib, OutCompressed.GetData() + ReservedBytes, CompressedSize, InBuffer.GetData(), InBuffer.Num());
+		if (!bSuccess)
+		{
+			//log error?
+			return false;
+		}
+#endif
 
 		Writer.Seek(0);
 		Writer << Size;
diff --git a/Engine/Plugins/Runtime/PacketHandlers/CompressionComponents/Oodle/Source/OodleHandlerComponent/Public/OodleArchives.h b/Engine/Plugins/Runtime/PacketHandlers/CompressionComponents/Oodle/Source/OodleHandlerComponent/Public/OodleArchives.h
--- a/Engine/Plugins/Runtime/PacketHandlers/CompressionComponents/Oodle/Source/OodleHandlerComponent/Public/OodleArchives.h
+++ b/Engine/Plugins/Runtime/PacketHandlers/CompressionComponents/Oodle/Source/OodleHandlerComponent/Public/OodleArchives.h
@@ -17,7 +17,7 @@
 #endif
 
 
-#if HAS_OODLE_SDK
+#if HAS_OODLE_NET_SDK
 
 #define CAPTURE_HEADER_MAGIC		0x41091CC4
 #define CAPTURE_FILE_VERSION		0x00000002
diff --git a/Engine/Plugins/Runtime/PacketHandlers/CompressionComponents/Oodle/Source/OodleHandlerComponent/Public/OodleHandlerComponent.h b/Engine/Plugins/Runtime/PacketHandlers/CompressionComponents/Oodle/Source/OodleHandlerComponent/Public/OodleHandlerComponent.h
--- a/Engine/Plugins/Runtime/PacketHandlers/CompressionComponents/Oodle/Source/OodleHandlerComponent/Public/OodleHandlerComponent.h
+++ b/Engine/Plugins/Runtime/PacketHandlers/CompressionComponents/Oodle/Source/OodleHandlerComponent/Public/OodleHandlerComponent.h
@@ -32,9 +32,14 @@ enum class EOodleEnableMode : uint8
 };
 
 
-#if HAS_OODLE_SDK
-
-#if UE4_OODLE_VER >= 200
+#if HAS_OODLE_NET_SDK
+
+#if UE4_OODLE_VER >= 270
+#include "oodle2net.h"
+	#if HAS_OODLE_DATA_SDK // allow data SDK usage if we have it
+	#include "oodle2.h"
+	#endif
+#elif UE4_OODLE_VER >= 200
 #include "oodle2.h"
 #else
 #include "oodle.h"
@@ -430,12 +435,9 @@ public:
 class FOodleComponentModuleInterface : public FPacketHandlerComponentModuleInterface
 {
 private:
-	/** Reference to the Oodle library handle */
-	void* OodleDllHandle;
 
 public:
 	FOodleComponentModuleInterface()
-		: OodleDllHandle(nullptr)
 	{
 	}
 
diff --git a/Engine/Plugins/Runtime/PacketHandlers/CompressionComponents/Oodle/Source/OodleHandlerComponent/Public/OodleUtils.h b/Engine/Plugins/Runtime/PacketHandlers/CompressionComponents/Oodle/Source/OodleHandlerComponent/Public/OodleUtils.h
--- a/Engine/Plugins/Runtime/PacketHandlers/CompressionComponents/Oodle/Source/OodleHandlerComponent/Public/OodleUtils.h
+++ b/Engine/Plugins/Runtime/PacketHandlers/CompressionComponents/Oodle/Source/OodleHandlerComponent/Public/OodleUtils.h
@@ -4,7 +4,7 @@
 
 #include "CoreMinimal.h"
 
-#if HAS_OODLE_SDK
+#if HAS_OODLE_NET_SDK
 
 namespace OodleUtils
 {
