Oodle archive minor refactor.

diff --git a/Engine/Plugins/Runtime/PacketHandlers/CompressionComponents/Oodle/Source/OodleHandlerComponent/Private/OodleArchives.cpp b/Engine/Plugins/Runtime/PacketHandlers/CompressionComponents/Oodle/Source/OodleHandlerComponent/Private/OodleArchives.cpp
--- a/Engine/Plugins/Runtime/PacketHandlers/CompressionComponents/Oodle/Source/OodleHandlerComponent/Private/OodleArchives.cpp
+++ b/Engine/Plugins/Runtime/PacketHandlers/CompressionComponents/Oodle/Source/OodleHandlerComponent/Private/OodleArchives.cpp
@@ -99,46 +99,42 @@ bool FOodleArchiveBase::SerializeOodleDecompressData(FOodleCompressedData& DataI
 	{
 		SeekPush(DataOffset);
 
-		if (CompressedLength == DecompressedLength)   //wasn't actually compressed
+		uint8* DecompressedData = new uint8[DecompressedLength];
+
+		if (CompressedLength == DecompressedLength)
 		{
-			uint8 * RawData = new uint8[DecompressedLength];
-			InnerArchive.Serialize(RawData, DecompressedLength);
-			OutData = RawData;
-			OutDataBytes = DecompressedLength;
+			InnerArchive.Serialize(DecompressedData, DecompressedLength);
 		}
+#ifdef HAS_OODLE_DATA_SDK
 		else
 		{
-#if HAS_OODLE_DATA_SDK
 			// @todo #JohnB: Remove bOutDataSlack after Oodle update, and after checking with Luigi
 			uint8* CompressedData = new uint8[CompressedLength + (bOutDataSlack ? OODLE_DICTIONARY_SLACK : 0)];
-			uint8* DecompressedData = new uint8[DecompressedLength];
-
+	
+	
 			// @todo #JohnB: Remove after Oodle update, and after checking with Luigi
 			if (bOutDataSlack)
 			{
 				FMemory::Memzero(CompressedData + CompressedLength, OODLE_DICTIONARY_SLACK);
 			}
-
+	
 			InnerArchive.Serialize(CompressedData, CompressedLength);
-
 			SINTa OodleLen = OodleLZ_Decompress((void*)CompressedData, CompressedLength, (void*)DecompressedData, DecompressedLength, OodleLZ_FuzzSafe_Yes);
-
-			bSuccess = OodleLen != 0;
-
-			if (bSuccess)
-			{
-				check(OodleLen == DecompressedLength);
-
-				OutData = DecompressedData;
-				OutDataBytes = DecompressedLength;
-			}
-			else
-			{
-				delete[] DecompressedData;
-			}
-
+			bSuccess = OodleLen == DecompressedLength;
+	
+	
 			delete[] CompressedData;
+		}
 #endif
+
+		if (bSuccess)
+		{
+			OutData = DecompressedData;
+			OutDataBytes = DecompressedLength;
+		}
+		else
+		{
+			delete[] DecompressedData;
 		}
 
 		SeekPop();
