Added QA-only Oodle clientside analytics.

diff --git a/Engine/Plugins/Runtime/PacketHandlers/CompressionComponents/Oodle/Source/OodleHandlerComponent/Private/OodleAnalytics.cpp b/Engine/Plugins/Runtime/PacketHandlers/CompressionComponents/Oodle/Source/OodleHandlerComponent/Private/OodleAnalytics.cpp
--- a/Engine/Plugins/Runtime/PacketHandlers/CompressionComponents/Oodle/Source/OodleHandlerComponent/Private/OodleAnalytics.cpp
+++ b/Engine/Plugins/Runtime/PacketHandlers/CompressionComponents/Oodle/Source/OodleHandlerComponent/Private/OodleAnalytics.cpp
@@ -191,7 +191,6 @@ void FOodleNetAnalyticsData::SendAnalytics()
 		UE_LOG(OodleHandlerComponentLog, Log, TEXT(" - OutAttemptedSavingsWithOverheadPercentTotal: %i"), OutAttemptedSavingsWithOverheadPercentTotal);
 
 
-		static const FString EZEventName = TEXT("Oodle.Stats");
 		static const FString EZAttrib_InCompressedNum = TEXT("InCompressedNum");
 		static const FString EZAttrib_InNotCompressedNum = TEXT("InNotCompressedNum");
 		static const FString EZAttrib_InCompressedWithOverheadLengthTotal = TEXT("InCompressedWithOverheadLengthTotal");
@@ -258,6 +257,16 @@ void FOodleNetAnalyticsData::SendAnalytics()
 			EZAttrib_OutAttemptedSavingsWithOverheadPercentTotal, OutAttemptedSavingsWithOverheadPercentTotal
 		);
 
-		AnalyticsProvider->RecordEvent(EZEventName, EventAttributes);
+		AnalyticsProvider->RecordEvent(GetAnalyticsEventName(), EventAttributes);
 	}
 }
+
+const TCHAR* FOodleNetAnalyticsData::GetAnalyticsEventName() const
+{
+	return TEXT("Oodle.Stats");
+}
+
+const TCHAR* FClientOodleNetAnalyticsData::GetAnalyticsEventName() const
+{
+	return TEXT("Oodle.ClientStats");
+}
diff --git a/Engine/Plugins/Runtime/PacketHandlers/CompressionComponents/Oodle/Source/OodleHandlerComponent/Private/OodleHandlerComponent.cpp b/Engine/Plugins/Runtime/PacketHandlers/CompressionComponents/Oodle/Source/OodleHandlerComponent/Private/OodleHandlerComponent.cpp
--- a/Engine/Plugins/Runtime/PacketHandlers/CompressionComponents/Oodle/Source/OodleHandlerComponent/Private/OodleHandlerComponent.cpp
+++ b/Engine/Plugins/Runtime/PacketHandlers/CompressionComponents/Oodle/Source/OodleHandlerComponent/Private/OodleHandlerComponent.cpp
@@ -1267,7 +1267,16 @@ void OodleHandlerComponent::NotifyAnalyticsProvider()
 
 		if (Handler->GetProvider().IsValid() && Aggregator.IsValid())
 		{
-			NetAnalyticsData = REGISTER_NET_ANALYTICS(Aggregator, FOodleNetAnalyticsData, TEXT("Oodle.Stats"));
+			const bool bIsServer = (Handler->Mode == Handler::Mode::Server);
+
+			if (bIsServer)
+			{
+				NetAnalyticsData = REGISTER_NET_ANALYTICS(Aggregator, FOodleNetAnalyticsData, TEXT("Oodle.Stats"));
+			}
+			else
+			{
+				NetAnalyticsData = REGISTER_NET_ANALYTICS(Aggregator, FClientOodleNetAnalyticsData, TEXT("Oodle.ClientStats"));
+			}
 		}
 	}
 
diff --git a/Engine/Plugins/Runtime/PacketHandlers/CompressionComponents/Oodle/Source/OodleHandlerComponent/Public/OodleAnalytics.h b/Engine/Plugins/Runtime/PacketHandlers/CompressionComponents/Oodle/Source/OodleHandlerComponent/Public/OodleAnalytics.h
--- a/Engine/Plugins/Runtime/PacketHandlers/CompressionComponents/Oodle/Source/OodleHandlerComponent/Public/OodleAnalytics.h
+++ b/Engine/Plugins/Runtime/PacketHandlers/CompressionComponents/Oodle/Source/OodleHandlerComponent/Public/OodleAnalytics.h
@@ -89,7 +89,7 @@ public:
 };
 
 /**
- * Oodle implementation for threaded net analytics data - the threading is taken care of, just need to send off the analytics
+ * Oodle implementation for (serverside) threaded net analytics data - threading is taken care of, just need to send off the analytics
  */
 struct FOodleNetAnalyticsData :
 #if NET_ANALYTICS_MULTITHREADING
@@ -108,5 +108,21 @@ public:
 		return this;
 	}
 #endif
+
+protected:
+	/**
+	 * Returns the analytics event name to use - to be overridden in subclasses
+	 */
+	virtual const TCHAR* GetAnalyticsEventName() const;
+};
+
+/**
+ * Clientside version of the above net analytics data (typically only used for debugging)
+ */
+struct FClientOodleNetAnalyticsData : public FOodleNetAnalyticsData
+{
+	virtual const TCHAR* GetAnalyticsEventName() const override;
 };
 
+
+
diff --git a/Engine/Source/Runtime/Engine/Private/NetworkDriver.cpp b/Engine/Source/Runtime/Engine/Private/NetworkDriver.cpp
--- a/Engine/Source/Runtime/Engine/Private/NetworkDriver.cpp
+++ b/Engine/Source/Runtime/Engine/Private/NetworkDriver.cpp
@@ -2248,6 +2248,11 @@ void UNetDriver::SetAnalyticsProvider(TSharedPtr<IAnalyticsProvider> InProvider)
 		ConnectionlessHandler->NotifyAnalyticsProvider(AnalyticsProvider, AnalyticsAggregator);
 	}
 
+	if (ServerConnection != nullptr)
+	{
+		ServerConnection->NotifyAnalyticsProvider();
+	}
+
 	for (UNetConnection* CurConn : ClientConnections)
 	{
 		if (CurConn != nullptr)
