Added Oodle analytics for counting Oodle compression handler instances, and the number of handlers with compression enabled vs disabled.

diff --git a/Engine/Plugins/Runtime/PacketHandlers/CompressionComponents/Oodle/Source/OodleHandlerComponent/Private/OodleAnalytics.cpp b/Engine/Plugins/Runtime/PacketHandlers/CompressionComponents/Oodle/Source/OodleHandlerComponent/Private/OodleAnalytics.cpp
--- a/Engine/Plugins/Runtime/PacketHandlers/CompressionComponents/Oodle/Source/OodleHandlerComponent/Private/OodleAnalytics.cpp
+++ b/Engine/Plugins/Runtime/PacketHandlers/CompressionComponents/Oodle/Source/OodleHandlerComponent/Private/OodleAnalytics.cpp
@@ -32,6 +32,8 @@ FOodleAnalyticsVars::FOodleAnalyticsVars()
 	, OutBeforeCompressedLengthTotal(0)
 	, OutNotCompressedFailedLengthTotal(0)
 	, OutNotCompressedSkippedLengthTotal(0)
+	, NumOodleHandlers(0)
+	, NumOodleHandlersCompressionEnabled(0)
 {
 }
 
@@ -55,7 +57,9 @@ bool FOodleAnalyticsVars::operator == (const FOodleAnalyticsVars& A) const
 		A.OutCompressedLengthTotal == OutCompressedLengthTotal &&
 		A.OutBeforeCompressedLengthTotal == OutBeforeCompressedLengthTotal &&
 		A.OutNotCompressedFailedLengthTotal == OutNotCompressedFailedLengthTotal &&
-		A.OutNotCompressedSkippedLengthTotal == OutNotCompressedSkippedLengthTotal;
+		A.OutNotCompressedSkippedLengthTotal == OutNotCompressedSkippedLengthTotal &&
+		A.NumOodleHandlers == NumOodleHandlers &&
+		A.NumOodleHandlersCompressionEnabled == NumOodleHandlersCompressionEnabled;
 }
 
 void FOodleAnalyticsVars::CommitAnalytics(FOodleAnalyticsVars& AggregatedData)
@@ -79,6 +83,8 @@ void FOodleAnalyticsVars::CommitAnalytics(FOodleAnalyticsVars& AggregatedData)
 	AggregatedData.OutBeforeCompressedLengthTotal += OutBeforeCompressedLengthTotal;
 	AggregatedData.OutNotCompressedFailedLengthTotal += OutNotCompressedFailedLengthTotal;
 	AggregatedData.OutNotCompressedSkippedLengthTotal += OutNotCompressedSkippedLengthTotal;
+	AggregatedData.NumOodleHandlers += NumOodleHandlers;
+	AggregatedData.NumOodleHandlersCompressionEnabled += NumOodleHandlersCompressionEnabled;
 }
 
 
@@ -157,6 +163,9 @@ void FOodleNetAnalyticsData::SendAnalytics()
 			int8 OutAttemptedSavingsWithOverheadPercentTotal = (1.0 - ((double)OutPostAttemptedWithOverheadLengthTotal / (double)OutPreAttemptedLengthTotal)) * 100.0;
 
 
+		uint32 NumOodleHandlersCompressionDisabled = NumOodleHandlers - NumOodleHandlersCompressionEnabled;
+
+
 		UE_LOG(OodleHandlerComponentLog, Log, TEXT("Oodle Analytics:"));
 		UE_LOG(OodleHandlerComponentLog, Log, TEXT(" - InCompressedNum: %llu"), InCompressedNum);
 		UE_LOG(OodleHandlerComponentLog, Log, TEXT(" - InNotCompressedNum: %llu"), InNotCompressedNum);
@@ -189,6 +198,9 @@ void FOodleNetAnalyticsData::SendAnalytics()
 		UE_LOG(OodleHandlerComponentLog, Log, TEXT(" - OutSavingsWithOverheadBytesTotal: %lli"), OutSavingsWithOverheadBytesTotal);
 		UE_LOG(OodleHandlerComponentLog, Log, TEXT(" - InAttemptedSavingsWithOverheadPercentTotal: %i"), InAttemptedSavingsWithOverheadPercentTotal);
 		UE_LOG(OodleHandlerComponentLog, Log, TEXT(" - OutAttemptedSavingsWithOverheadPercentTotal: %i"), OutAttemptedSavingsWithOverheadPercentTotal);
+		UE_LOG(OodleHandlerComponentLog, Log, TEXT(" - NumOodleHandlers: %i"), NumOodleHandlers);
+		UE_LOG(OodleHandlerComponentLog, Log, TEXT(" - NumOodleHandlersCompressionEnabled: %i"), NumOodleHandlersCompressionEnabled);
+		UE_LOG(OodleHandlerComponentLog, Log, TEXT(" - NumOodleHandlersCompressionDisabled: %i"), NumOodleHandlersCompressionDisabled);
 
 
 		static const FString EZAttrib_InCompressedNum = TEXT("InCompressedNum");
@@ -222,6 +234,9 @@ void FOodleNetAnalyticsData::SendAnalytics()
 		static const FString EZAttrib_OutSavingsWithOverheadBytesTotal = TEXT("OutSavingsWithOverheadBytesTotal");
 		static const FString EZAttrib_InAttemptedSavingsWithOverheadPercentTotal = TEXT("InAttemptedSavingsWithOverheadPercentTotal");
 		static const FString EZAttrib_OutAttemptedSavingsWithOverheadPercentTotal = TEXT("OutAttemptedSavingsWithOverheadPercentTotal");
+		static const FString EZAttrib_NumOodleHandlers = TEXT("NumOodleHandlers");
+		static const FString EZAttrib_NumOodleHandlersCompressionEnabled = TEXT("NumOodleHandlersCompressionEnabled");
+		static const FString EZAttrib_NumOodleHandlersCompressionDisabled = TEXT("NumOodleHandlersCompressionDisabled");
 
 		const TArray<FAnalyticsEventAttribute> EventAttributes = MakeAnalyticsEventAttributeArray(
 			EZAttrib_InCompressedNum, InCompressedNum,
@@ -254,7 +269,10 @@ void FOodleNetAnalyticsData::SendAnalytics()
 			EZAttrib_InSavingsWithOverheadBytesTotal, InSavingsWithOverheadBytesTotal,
 			EZAttrib_OutSavingsWithOverheadBytesTotal, OutSavingsWithOverheadBytesTotal,
 			EZAttrib_InAttemptedSavingsWithOverheadPercentTotal, InAttemptedSavingsWithOverheadPercentTotal,
-			EZAttrib_OutAttemptedSavingsWithOverheadPercentTotal, OutAttemptedSavingsWithOverheadPercentTotal
+			EZAttrib_OutAttemptedSavingsWithOverheadPercentTotal, OutAttemptedSavingsWithOverheadPercentTotal,
+			EZAttrib_NumOodleHandlers, NumOodleHandlers,
+			EZAttrib_NumOodleHandlersCompressionEnabled, NumOodleHandlersCompressionEnabled,
+			EZAttrib_NumOodleHandlersCompressionDisabled, NumOodleHandlersCompressionDisabled
 		);
 
 		AnalyticsProvider->RecordEvent(GetAnalyticsEventName(), EventAttributes);
diff --git a/Engine/Plugins/Runtime/PacketHandlers/CompressionComponents/Oodle/Source/OodleHandlerComponent/Private/OodleHandlerComponent.cpp b/Engine/Plugins/Runtime/PacketHandlers/CompressionComponents/Oodle/Source/OodleHandlerComponent/Private/OodleHandlerComponent.cpp
--- a/Engine/Plugins/Runtime/PacketHandlers/CompressionComponents/Oodle/Source/OodleHandlerComponent/Private/OodleHandlerComponent.cpp
+++ b/Engine/Plugins/Runtime/PacketHandlers/CompressionComponents/Oodle/Source/OodleHandlerComponent/Private/OodleHandlerComponent.cpp
@@ -493,6 +493,18 @@ void OodleHandlerComponent::InitializeDictionaries()
 	}
 }
 
+void OodleHandlerComponent::RemoteInitializeDictionaries()
+{
+	InitializeDictionaries();
+
+	if (bOodleAnalytics && NetAnalyticsData.IsValid())
+	{
+		FOodleAnalyticsVars* AnalyticsVars = NetAnalyticsData->GetLocalData();
+
+		AnalyticsVars->NumOodleHandlersCompressionEnabled++;
+	}
+}
+
 void OodleHandlerComponent::InitializeDictionary(FString FilePath, TSharedPtr<FOodleDictionary>& OutDictionary)
 {
 	TSharedPtr<FOodleDictionary>* DictionaryRef = DictionaryMap.Find(FilePath);
@@ -873,7 +885,7 @@ void OodleHandlerComponent::Incoming(FBitReader& Packet)
 			// Lazy-loading of dictionary when EOodleEnableMode::WhenCompressedPacketReceived is active
 			if (!bInitializedDictionaries && (bIsServer ? ServerEnableMode : ClientEnableMode) == EOodleEnableMode::WhenCompressedPacketReceived)
 			{
-				InitializeDictionaries();
+				RemoteInitializeDictionaries();
 			}
 
 			FOodleDictionary* CurDict = (bIsServer ? ClientDictionary.Get() : ServerDictionary.Get());
@@ -1278,6 +1290,21 @@ void OodleHandlerComponent::NotifyAnalyticsProvider()
 				NetAnalyticsData = REGISTER_NET_ANALYTICS(Aggregator, FClientOodleNetAnalyticsData, TEXT("Oodle.ClientStats"));
 			}
 		}
+
+		// This depends upon NetAnalyticsData only being initialized once, so that this component is not counted more than once
+		if (NetAnalyticsData.IsValid())
+		{
+			if (FOodleAnalyticsVars* AnalyticsVars = NetAnalyticsData->GetLocalData())
+			{
+				AnalyticsVars->NumOodleHandlers++;
+
+				// If compression's already enabled, count it here - as InitializeDictionaries may happen before NetAnalyticsData is set
+				if (IsCompressionActive())
+				{
+					AnalyticsVars->NumOodleHandlersCompressionEnabled++;
+				}
+			}
+		}
 	}
 
 	bOodleAnalytics = NetAnalyticsData.IsValid();
diff --git a/Engine/Plugins/Runtime/PacketHandlers/CompressionComponents/Oodle/Source/OodleHandlerComponent/Public/OodleAnalytics.h b/Engine/Plugins/Runtime/PacketHandlers/CompressionComponents/Oodle/Source/OodleHandlerComponent/Public/OodleAnalytics.h
--- a/Engine/Plugins/Runtime/PacketHandlers/CompressionComponents/Oodle/Source/OodleHandlerComponent/Public/OodleAnalytics.h
+++ b/Engine/Plugins/Runtime/PacketHandlers/CompressionComponents/Oodle/Source/OodleHandlerComponent/Public/OodleAnalytics.h
@@ -86,6 +86,13 @@ public:
 
 	/** The length of all outgoing packets, where compression was skipped. */
 	uint64 OutNotCompressedSkippedLengthTotal;
+
+
+	/** The number of OodleHandlerComponent's running during the lifetime of the analytics aggregator (i.e. NetDriver lifetime). */
+	uint32 NumOodleHandlers;
+
+	/** The number of OodleHandlerComponent's that had packet compression enabled. */
+	uint32 NumOodleHandlersCompressionEnabled;
 };
 
 /**
diff --git a/Engine/Plugins/Runtime/PacketHandlers/CompressionComponents/Oodle/Source/OodleHandlerComponent/Public/OodleHandlerComponent.h b/Engine/Plugins/Runtime/PacketHandlers/CompressionComponents/Oodle/Source/OodleHandlerComponent/Public/OodleHandlerComponent.h
--- a/Engine/Plugins/Runtime/PacketHandlers/CompressionComponents/Oodle/Source/OodleHandlerComponent/Public/OodleHandlerComponent.h
+++ b/Engine/Plugins/Runtime/PacketHandlers/CompressionComponents/Oodle/Source/OodleHandlerComponent/Public/OodleHandlerComponent.h
@@ -303,6 +303,11 @@ public:
 	 */
 	void InitializeDictionaries();
 
+	/**
+	 * Lazy dictionary initialization, triggered by receiving a compressed packet from the remote connection
+	 */
+	void RemoteInitializeDictionaries();
+
 	/**
 	 * Initializes FOodleDictionary data, from the specified dictionary file
 	 *
