Added CVar versions of ClientEnableMode and ServerEnableMode for Oodle - net.OodleServerEnableMode and net.OodleClientEnableMode - to allow setting in device scalability settings.

diff --git a/Engine/Plugins/Runtime/PacketHandlers/CompressionComponents/Oodle/Source/OodleHandlerComponent/Private/OodleHandlerComponent.cpp b/Engine/Plugins/Runtime/PacketHandlers/CompressionComponents/Oodle/Source/OodleHandlerComponent/Private/OodleHandlerComponent.cpp
--- a/Engine/Plugins/Runtime/PacketHandlers/CompressionComponents/Oodle/Source/OodleHandlerComponent/Private/OodleHandlerComponent.cpp
+++ b/Engine/Plugins/Runtime/PacketHandlers/CompressionComponents/Oodle/Source/OodleHandlerComponent/Private/OodleHandlerComponent.cpp
@@ -100,6 +100,15 @@ TAutoConsoleVariable<int32> CVarOodleMinSizeForCompression(
 	TEXT("The minimum size an outgoing packet must be, for it to be considered for compression (does not count overhead of handler components which process packets after Oodle)."));
 
 
+TAutoConsoleVariable<FString> CVarOodleServerEnableMode(
+	TEXT("net.OodleServerEnableMode"), "",
+	TEXT("When to enable compression on the server (overrides the 'ServerEnableMode' .ini setting)."));
+
+TAutoConsoleVariable<FString> CVarOodleClientEnableMode(
+	TEXT("net.OodleClientEnableMode"), "",
+	TEXT("When to enable compression on the client (overrides the 'ClientEnableMode' .ini setting)."));
+
+
 
 // @todo #JohnB: Remove after Oodle update, and after checking with Luigi
 static rrbool STDCALL UEOodleDisplayAssert(const char* File, const int Line, const char* Function, const char* Message)
@@ -366,27 +375,51 @@ void OodleHandlerComponent::Initialize()
 
 	FString CurEnableModeStr;
 	UEnum* EnableModeEnum = StaticEnum<EOodleEnableMode>();
+	bool bSetServerEnableMode = false;
+	bool bSetClientEnableMode = false;
+	auto SetEnableModeFromStr = [&EnableModeEnum](EOodleEnableMode& OutMode, FString ModeStr) -> bool
+		{
+			bool bSuccess = false;
 
-	if (GConfig->GetString(OODLE_INI_SECTION, TEXT("ServerEnableMode"), CurEnableModeStr, GEngineIni))
-	{
-		int64 EnumVal = EnableModeEnum->GetValueByNameString(CurEnableModeStr);
+			int64 EnumVal = EnableModeEnum->GetValueByNameString(ModeStr);
 
-		if (EnumVal != INDEX_NONE)
-		{
-			ServerEnableMode = (EOodleEnableMode)EnumVal;
-		}
-	}
+			if (EnumVal != INDEX_NONE)
+			{
+				OutMode = (EOodleEnableMode)EnumVal;
+				bSuccess = true;
+			}
+			else
+			{
+				UE_LOG(OodleHandlerComponentLog, Error, TEXT("Failed to parse EOodleEnableMode value '%s'"), *ModeStr);
+			}
+
+			return bSuccess;
+		};
+
+	// ServerEnableMode
+	CurEnableModeStr = CVarOodleServerEnableMode.GetValueOnAnyThread();
 
-	if (GConfig->GetString(OODLE_INI_SECTION, TEXT("ClientEnableMode"), CurEnableModeStr, GEngineIni))
+	bSetServerEnableMode = CurEnableModeStr.Len() > 0 && SetEnableModeFromStr(ServerEnableMode, CurEnableModeStr);
+
+	if (!bSetServerEnableMode)
 	{
-		int64 EnumVal = EnableModeEnum->GetValueByNameString(CurEnableModeStr);
+		bSetServerEnableMode = GConfig->GetString(OODLE_INI_SECTION, TEXT("ServerEnableMode"), CurEnableModeStr, GEngineIni) &&
+								SetEnableModeFromStr(ServerEnableMode, CurEnableModeStr);
+	}
 
-		if (EnumVal != INDEX_NONE)
-		{
-			ClientEnableMode = (EOodleEnableMode)EnumVal;
-		}
+	// ClientEnableMode
+	CurEnableModeStr = CVarOodleClientEnableMode.GetValueOnAnyThread();
+
+	bSetClientEnableMode = CurEnableModeStr.Len() > 0 && SetEnableModeFromStr(ClientEnableMode, CurEnableModeStr);
+
+	if (!bSetClientEnableMode)
+	{
+		bSetClientEnableMode = GConfig->GetString(OODLE_INI_SECTION, TEXT("ClientEnableMode"), CurEnableModeStr, GEngineIni) &&
+								SetEnableModeFromStr(ClientEnableMode, CurEnableModeStr);
 	}
 
+
+
 #if !UE_BUILD_SHIPPING || OODLE_DEV_SHIPPING
 	GConfig->GetBool(OODLE_INI_SECTION, TEXT("bUseDictionaryIfPresent"), bUseDictionaryIfPresent, GEngineIni);
 
