Add missing include to OodleCompressionFormat

diff --git a/Engine/Plugins/Runtime/PacketHandlers/CompressionComponents/Oodle/Source/OodleCompressionFormat/Private/OodleCompressionFormat.cpp b/Engine/Plugins/Runtime/PacketHandlers/CompressionComponents/Oodle/Source/OodleCompressionFormat/Private/OodleCompressionFormat.cpp
--- a/Engine/Plugins/Runtime/PacketHandlers/CompressionComponents/Oodle/Source/OodleCompressionFormat/Private/OodleCompressionFormat.cpp
+++ b/Engine/Plugins/Runtime/PacketHandlers/CompressionComponents/Oodle/Source/OodleCompressionFormat/Private/OodleCompressionFormat.cpp
@@ -4,6 +4,7 @@
 #include "Misc/Compression.h"
 #include "Misc/ICompressionFormat.h"
 #include "Misc/CommandLine.h"
+#include "Misc/Parse.h"
 
 #if HAS_OODLE_DATA_SDK
 
