Added .ini options for default Oodle Dictionary Generation settings - these can still be overriden from the commandline.

diff --git a/Engine/Plugins/Runtime/PacketHandlers/CompressionComponents/Oodle/Source/OodleHandlerComponent/Classes/OodleTrainerCommandlet.h b/Engine/Plugins/Runtime/PacketHandlers/CompressionComponents/Oodle/Source/OodleHandlerComponent/Classes/OodleTrainerCommandlet.h
--- a/Engine/Plugins/Runtime/PacketHandlers/CompressionComponents/Oodle/Source/OodleHandlerComponent/Classes/OodleTrainerCommandlet.h
+++ b/Engine/Plugins/Runtime/PacketHandlers/CompressionComponents/Oodle/Source/OodleHandlerComponent/Classes/OodleTrainerCommandlet.h
@@ -60,12 +60,43 @@
  *		- Outputs information about the packet file, such as the MB amount of data recorded, per net connection channel, and data types
  *		- @todo #JohnB: Only implement, if deciding to actually capture/track this kind of data
  */
-UCLASS()
+UCLASS(config=Editor)
 class UOodleTrainerCommandlet : public UCommandlet
 {
-	GENERATED_UCLASS_BODY()
+	GENERATED_BODY()
 
 public:
+	/** Whether or not compression testing should be performed after dictionary generation (uses up some of the packets) */
+	UPROPERTY(config)
+	bool bCompressionTest;
+
+	/** Size of the hash table to use for the dictionary */
+	UPROPERTY(config)
+	int32 HashTableSize;
+
+	/** Size of the dictionary to be generated */
+	UPROPERTY(config)
+	int32 DictionarySize;
+
+	/** The number of random packet-selection trials to run, when generating the dictionary, to try and optimize the dictionary */
+	UPROPERTY(config)
+	int32 DictionaryTrials;
+
+	/** The randomness, in percent, of random packet-selection trials */
+	UPROPERTY(config)
+	int32 TrialRandomness;
+
+	/** The number of generations of random packet-selection trials  */
+	UPROPERTY(config)
+	int32 TrialGenerations;
+
+	/** Whether or not random-trials have been disabled */
+	UPROPERTY(config)
+	bool bNoTrials;
+
+public:
+	UOodleTrainerCommandlet(const FObjectInitializer& ObjectInitializer);
+
 	virtual int32 Main(const FString& Params) override;
 
 #if USE_OODLE_TRAINER_COMMANDLET
@@ -227,11 +258,11 @@ public:
 	FOodleDictionaryGenerator()
 		: OutputDictionaryFile()
 		, bCompressionTest(false)
-		, HashTableSize(19)
-		, DictionarySize(4 * 1024 * 1024)
-		, DictionaryTrials(3)
-		, TrialRandomness(200)
-		, TrialGenerations(1)
+		, HashTableSize(0)
+		, DictionarySize(0)
+		, DictionaryTrials(0)
+		, TrialRandomness(0)
+		, TrialGenerations(0)
 		, bNoTrials(false)
 		, MergeMap()
 		, DictionaryPacketBytes(0)
diff --git a/Engine/Plugins/Runtime/PacketHandlers/CompressionComponents/Oodle/Source/OodleHandlerComponent/Private/OodleTrainerCommandlet.cpp b/Engine/Plugins/Runtime/PacketHandlers/CompressionComponents/Oodle/Source/OodleHandlerComponent/Private/OodleTrainerCommandlet.cpp
--- a/Engine/Plugins/Runtime/PacketHandlers/CompressionComponents/Oodle/Source/OodleHandlerComponent/Private/OodleTrainerCommandlet.cpp
+++ b/Engine/Plugins/Runtime/PacketHandlers/CompressionComponents/Oodle/Source/OodleHandlerComponent/Private/OodleTrainerCommandlet.cpp
@@ -29,6 +29,13 @@
 
 UOodleTrainerCommandlet::UOodleTrainerCommandlet(const FObjectInitializer& ObjectInitializer)
 	: Super(ObjectInitializer)
+	, bCompressionTest(false)
+	, HashTableSize(19)
+	, DictionarySize(4 * 1024 * 1024)
+	, DictionaryTrials(3)
+	, TrialRandomness(200)
+	, TrialGenerations(1)
+	, bNoTrials(false)
 {
 	LogToConsole = true;
 }
@@ -611,9 +618,17 @@ bool FOodleDictionaryGenerator::InitGenerator()
 	bDebugDump = FParse::Param(FCommandLine::Get(), TEXT("OodleDebugDump"));
 
 	// Parse and verify parameters
-	// @todo #JohnB: Set these defaults from the .ini file as well (commandline overriding) - mandatory for tailoring to specific games,
-	//					using the games .ini file
+	const UOodleTrainerCommandlet* CDO = GetDefault<UOodleTrainerCommandlet>();
 
+	HashTableSize = CDO->HashTableSize;
+	DictionarySize = CDO->DictionarySize;
+	DictionaryTrials = CDO->DictionaryTrials;
+	TrialRandomness = CDO->TrialRandomness;
+	TrialGenerations = CDO->TrialGenerations;
+	bNoTrials = CDO->bNoTrials;
+	bCompressionTest = CDO->bCompressionTest;
+
+	// Commandline overrides
 	FParse::Value(FCommandLine::Get(), TEXT("-HashTableSize="), HashTableSize);
 	FParse::Value(FCommandLine::Get(), TEXT("-DictionarySize="), DictionarySize);
 	FParse::Value(FCommandLine::Get(), TEXT("-DictionaryTrials="), DictionaryTrials);
