Enable fuzz-safe Oodle decompression to prevent crashes on I/O failures

diff --git a/Engine/Plugins/Runtime/PacketHandlers/CompressionComponents/Oodle/Source/OodleCompressionFormat/Private/OodleCompressionFormat.cpp b/Engine/Plugins/Runtime/PacketHandlers/CompressionComponents/Oodle/Source/OodleCompressionFormat/Private/OodleCompressionFormat.cpp
--- a/Engine/Plugins/Runtime/PacketHandlers/CompressionComponents/Oodle/Source/OodleCompressionFormat/Private/OodleCompressionFormat.cpp
+++ b/Engine/Plugins/Runtime/PacketHandlers/CompressionComponents/Oodle/Source/OodleCompressionFormat/Private/OodleCompressionFormat.cpp
@@ -132,7 +132,7 @@ struct FOodleCustomCompressor : ICompressionFormat
 	{
 		ConditionalInitialize();
 
-		int32 Result = (int32)OodleLZ_Decompress(CompressedBuffer, CompressedSize, UncompressedBuffer, UncompressedSize, OodleLZ_FuzzSafe_No);
+		int32 Result = (int32)OodleLZ_Decompress(CompressedBuffer, CompressedSize, UncompressedBuffer, UncompressedSize, OodleLZ_FuzzSafe_Yes);
 		if (Result > 0)
 		{
 			UncompressedSize = Result;
