- Added ability for a project to choose oodle version (only 280 and 255 are there now), using this section in DefaultEngine.ini:

[OodleSdk]
OverrideSdkVersion=255

diff --git a/Engine/Plugins/Runtime/PacketHandlers/CompressionComponents/Oodle/Source/OodleCompressionFormat/Private/OodleCompressionFormat.cpp b/Engine/Plugins/Runtime/PacketHandlers/CompressionComponents/Oodle/Source/OodleCompressionFormat/Private/OodleCompressionFormat.cpp
--- a/Engine/Plugins/Runtime/PacketHandlers/CompressionComponents/Oodle/Source/OodleCompressionFormat/Private/OodleCompressionFormat.cpp
+++ b/Engine/Plugins/Runtime/PacketHandlers/CompressionComponents/Oodle/Source/OodleCompressionFormat/Private/OodleCompressionFormat.cpp
@@ -145,7 +145,11 @@ struct FOodleCustomCompressor : ICompressionFormat
 	{
 		ConditionalInitialize();
 
+#if UE4_OODLE_VER > 255
 		int32 Needed = (int32)OodleLZ_GetCompressedBufferSizeNeeded(Compressor, UncompressedSize);
+#else
+		int32 Needed = (int32)OodleLZ_GetCompressedBufferSizeNeeded(UncompressedSize);
+#endif
 		return Needed;
 	}
 };
@@ -203,7 +207,9 @@ class FOodleCompressionFormatModuleInterface : public IModuleInterface
 
 		IModularFeatures::Get().RegisterModularFeature(COMPRESSION_FORMAT_FEATURE_NAME, CompressionFormat);
 
+#if UE4_OODLE_VER > 255
 		OodleCore_Plugins_SetPrintf(OodleCore_Plugin_Printf_Verbose);
+#endif
 
 #endif
 	}
diff --git a/Engine/Plugins/Runtime/PacketHandlers/CompressionComponents/Oodle/Source/OodleHandlerComponent/Private/OodleArchives.cpp b/Engine/Plugins/Runtime/PacketHandlers/CompressionComponents/Oodle/Source/OodleHandlerComponent/Private/OodleArchives.cpp
--- a/Engine/Plugins/Runtime/PacketHandlers/CompressionComponents/Oodle/Source/OodleHandlerComponent/Private/OodleArchives.cpp
+++ b/Engine/Plugins/Runtime/PacketHandlers/CompressionComponents/Oodle/Source/OodleHandlerComponent/Private/OodleArchives.cpp
@@ -43,7 +43,11 @@ bool FOodleArchiveBase::SerializeOodleCompressData(FOodleCompressedData& OutData
 		OutDataInfo.DecompressedLength.Set(*this, DataBytes);
 	
 #if HAS_OODLE_DATA_SDK
+#if UE4_OODLE_VER > 255
 		uint32 CompressBufferLen = OodleLZ_GetCompressedBufferSizeNeeded(OodleLZ_Compressor_Kraken, DataBytes);
+#else
+		uint32 CompressBufferLen = OodleLZ_GetCompressedBufferSizeNeeded(DataBytes);
+#endif
 		uint8* CompressBuffer = new uint8[CompressBufferLen];
 
 		SINTa OodleLen = OodleLZ_Compress(OodleLZ_Compressor_Kraken, Data, (SINTa)DataBytes, CompressBuffer,
@@ -371,7 +375,6 @@ void FOodleDictionaryArchive::FDictionaryHeader::SerializeHeader(FOodleDictionar
 
 		bSuccess = bSuccess && ensure(Magic == DICTIONARY_HEADER_MAGIC);
 		bSuccess = bSuccess && ensure(DictionaryVersion <= DICTIONARY_FILE_VERSION);
-		//bSuccess = bSuccess && ensure(OodleMajorHeaderVersion == OODLE2NET_VERSION_MAJOR);   //deprecated. oodle network data format has stabilised.
 	}
 
 	if (!bSuccess)
@@ -392,7 +395,11 @@ FOodleDictionaryArchive::FOodleDictionaryArchive(FArchive& InInnerArchive)
 	if (IsSaving())
 	{
 		Header.DictionaryVersion = DICTIONARY_FILE_VERSION;
+#if UE4_OODLE_VER > 255
 		Header.OodleMajorHeaderVersion = OODLE2NET_VERSION_MAJOR;
+#else
+		Header.OodleMajorHeaderVersion = OODLE2_VERSION_MAJOR;
+#endif
 	}
 }
 
diff --git a/Engine/Plugins/Runtime/PacketHandlers/CompressionComponents/Oodle/Source/OodleHandlerComponent/Private/OodleUtils.cpp b/Engine/Plugins/Runtime/PacketHandlers/CompressionComponents/Oodle/Source/OodleHandlerComponent/Private/OodleUtils.cpp
--- a/Engine/Plugins/Runtime/PacketHandlers/CompressionComponents/Oodle/Source/OodleHandlerComponent/Private/OodleUtils.cpp
+++ b/Engine/Plugins/Runtime/PacketHandlers/CompressionComponents/Oodle/Source/OodleHandlerComponent/Private/OodleUtils.cpp
@@ -47,7 +47,11 @@ namespace OodleUtils
 		const OodleLZ_Compressor			Compressor = OodleLZ_Compressor_Selkie;
 		const OodleLZ_CompressionLevel		Level = OodleLZ_CompressionLevel_VeryFast;
 
+#if UE4_OODLE_VER > 255
 		OutCompressed.SetNum(ReservedBytes + OodleLZ_GetCompressedBufferSizeNeeded(Compressor, InBuffer.Num()));
+#else
+		OutCompressed.SetNum(ReservedBytes + OodleLZ_GetCompressedBufferSizeNeeded(InBuffer.Num()));
+#endif
 		CompressedSize = OodleLZ_Compress(Compressor, InBuffer.GetData(), (SINTa)InBuffer.Num(), OutCompressed.GetData() + ReservedBytes, Level);
 
 		if (CompressedSize == OODLELZ_FAILED)
