- Fixed fuzz safety error with oodle 255

diff --git a/Engine/Plugins/Runtime/PacketHandlers/CompressionComponents/Oodle/Source/OodleHandlerComponent/Private/OodleArchives.cpp b/Engine/Plugins/Runtime/PacketHandlers/CompressionComponents/Oodle/Source/OodleHandlerComponent/Private/OodleArchives.cpp
--- a/Engine/Plugins/Runtime/PacketHandlers/CompressionComponents/Oodle/Source/OodleHandlerComponent/Private/OodleArchives.cpp
+++ b/Engine/Plugins/Runtime/PacketHandlers/CompressionComponents/Oodle/Source/OodleHandlerComponent/Private/OodleArchives.cpp
@@ -123,7 +123,11 @@ bool FOodleArchiveBase::SerializeOodleDecompressData(FOodleCompressedData& DataI
 			}
 	
 			InnerArchive.Serialize(CompressedData, CompressedLength);
+#if UE4_OODLE_VER > 255
 			SINTa OodleLen = OodleLZ_Decompress((void*)CompressedData, CompressedLength, (void*)DecompressedData, DecompressedLength, OodleLZ_FuzzSafe_Yes);
+#else
+			SINTa OodleLen = OodleLZ_Decompress((void*)CompressedData, CompressedLength, (void*)DecompressedData, DecompressedLength, OodleLZ_FuzzSafe_No);
+#endif
 			bSuccess = OodleLen == DecompressedLength;
 	
 	
