Change OodleHandlerComponent max packet bytes from check to static_assert

diff --git a/Engine/Plugins/Runtime/PacketHandlers/CompressionComponents/Oodle/Source/OodleHandlerComponent/Private/OodleHandlerComponent.cpp b/Engine/Plugins/Runtime/PacketHandlers/CompressionComponents/Oodle/Source/OodleHandlerComponent/Private/OodleHandlerComponent.cpp
--- a/Engine/Plugins/Runtime/PacketHandlers/CompressionComponents/Oodle/Source/OodleHandlerComponent/Private/OodleHandlerComponent.cpp
+++ b/Engine/Plugins/Runtime/PacketHandlers/CompressionComponents/Oodle/Source/OodleHandlerComponent/Private/OodleHandlerComponent.cpp
@@ -930,11 +930,6 @@ void OodleHandlerComponent::Incoming(FBitReader& Packet)
 
 				SerializeOodlePacketSize(Packet, DecompressedLength);
 
-#if !UE_BUILD_SHIPPING
-				// Never allow DecompressedLength values bigger than this, due to performance/security considerations
-				check(MAX_OODLE_PACKET_BYTES <= 16384);
-#endif
-
 				if (DecompressedLength < MAX_OODLE_PACKET_BYTES)
 				{
 					static uint8 CompressedData[MAX_OODLE_BUFFER];
diff --git a/Engine/Plugins/Runtime/PacketHandlers/CompressionComponents/Oodle/Source/OodleHandlerComponent/Public/OodleHandlerComponent.h b/Engine/Plugins/Runtime/PacketHandlers/CompressionComponents/Oodle/Source/OodleHandlerComponent/Public/OodleHandlerComponent.h
--- a/Engine/Plugins/Runtime/PacketHandlers/CompressionComponents/Oodle/Source/OodleHandlerComponent/Public/OodleHandlerComponent.h
+++ b/Engine/Plugins/Runtime/PacketHandlers/CompressionComponents/Oodle/Source/OodleHandlerComponent/Public/OodleHandlerComponent.h
@@ -20,6 +20,8 @@ DECLARE_LOG_CATEGORY_EXTERN(OodleHandlerComponentLog, Log, All);
 // The maximum compress/decompress buffer size - overkill, as buffers are statically allocated, and can't use Oodle runtime buffer calc
 #define MAX_OODLE_BUFFER	(MAX_OODLE_PACKET_BYTES * 2)
 
+static_assert(MAX_OODLE_PACKET_BYTES <= 16384, "Never allow DecompressedLength values bigger than this, due to performance/security considerations");
+
 /**
  * Specifies when compression is enabled. Used to make compression optional, for some platforms/clients
  */
