Make OodleComponentHandler destroy proxy archives before inner archive

diff --git a/Engine/Plugins/Runtime/PacketHandlers/CompressionComponents/Oodle/Source/OodleHandlerComponent/Private/OodleHandlerComponent.cpp b/Engine/Plugins/Runtime/PacketHandlers/CompressionComponents/Oodle/Source/OodleHandlerComponent/Private/OodleHandlerComponent.cpp
--- a/Engine/Plugins/Runtime/PacketHandlers/CompressionComponents/Oodle/Source/OodleHandlerComponent/Private/OodleHandlerComponent.cpp
+++ b/Engine/Plugins/Runtime/PacketHandlers/CompressionComponents/Oodle/Source/OodleHandlerComponent/Private/OodleHandlerComponent.cpp
@@ -545,7 +545,7 @@ void OodleHandlerComponent::InitializeDictionary(FString FilePath, TSharedPtr<FO
 	// Load the dictionary, if it's not yet loaded
 	if (DictionaryRef == nullptr)
 	{
-		FArchive* ReadArc = IFileManager::Get().CreateFileReader(*FilePath);
+		TUniquePtr<FArchive> ReadArc(IFileManager::Get().CreateFileReader(*FilePath));
 
 		if (ReadArc != nullptr)
 		{
@@ -607,9 +607,6 @@ void OodleHandlerComponent::InitializeDictionary(FString FilePath, TSharedPtr<FO
 
 
 			ReadArc->Close();
-
-			delete ReadArc;
-			ReadArc = nullptr;
 		}
 		else
 		{
@@ -875,7 +872,9 @@ void OodleHandlerComponent::FreePacketLogs()
 	if (OutPacketLog != nullptr)
 	{
 		OutPacketLog->Close();
-		OutPacketLog->DeleteInnerArchive();
+
+		// Proxy archives must be deleted before their inner archive
+		TUniquePtr<FArchive> InnerDeleter(&OutPacketLog->GetInnerArchive());
 
 		delete OutPacketLog;
 
@@ -885,7 +884,9 @@ void OodleHandlerComponent::FreePacketLogs()
 	if (InPacketLog != nullptr)
 	{
 		InPacketLog->Close();
-		InPacketLog->DeleteInnerArchive();
+
+		// Proxy archives must be deleted before their inner archive
+		TUniquePtr<FArchive> InnerDeleter(&InPacketLog->GetInnerArchive());
 
 		delete InPacketLog;
 
diff --git a/Engine/Plugins/Runtime/PacketHandlers/CompressionComponents/Oodle/Source/OodleHandlerComponent/Public/OodleArchives.h b/Engine/Plugins/Runtime/PacketHandlers/CompressionComponents/Oodle/Source/OodleHandlerComponent/Public/OodleArchives.h
--- a/Engine/Plugins/Runtime/PacketHandlers/CompressionComponents/Oodle/Source/OodleHandlerComponent/Public/OodleArchives.h
+++ b/Engine/Plugins/Runtime/PacketHandlers/CompressionComponents/Oodle/Source/OodleHandlerComponent/Public/OodleArchives.h
@@ -173,15 +173,11 @@ public:
 		bImmediateFlush = FParse::Param(FCommandLine::Get(), TEXT("FORCELOGFLUSH"));
 	}
 
-	/**
-	 * Deletes the inner archive
-	 */
-	void DeleteInnerArchive()
+	FArchive& GetInnerArchive()
 	{
-		delete (FArchive*)&InnerArchive;
+		return InnerArchive;
 	}
 
-
 	/**
 	 * Pushes the current archive position onto a stack
 	 */
