Remove Oodle UE4 references.

diff --git a/Engine/Plugins/Runtime/PacketHandlers/CompressionComponents/Oodle/BatchDebugDump.bat b/Engine/Plugins/Runtime/PacketHandlers/CompressionComponents/Oodle/BatchDebugDump.bat
--- a/Engine/Plugins/Runtime/PacketHandlers/CompressionComponents/Oodle/BatchDebugDump.bat
+++ b/Engine/Plugins/Runtime/PacketHandlers/CompressionComponents/Oodle/BatchDebugDump.bat
@@ -7,18 +7,18 @@ echo.
 REM This batch file can only work from within the Oodle folder.
 set BaseFolder="..\..\..\..\..\.."
 
-if exist %BaseFolder:"=%\Engine goto SetUE4Editor
+if exist %BaseFolder:"=%\Engine goto SetUEEditor
 
-echo Could not locate Engine folder. This .bat must be located within UE4\Engine\Plugins\Runtime\PacketHandlers\CompressionComponents\Oodle
+echo Could not locate Engine folder. This .bat must be located within UE\Engine\Plugins\Runtime\PacketHandlers\CompressionComponents\Oodle
 goto End
 
 
-:SetUE4Editor
-set UE4EditorLoc="%BaseFolder:"=%\Engine\Binaries\Win64\UE4Editor.exe"
+:SetUEEditor
+set UEEditorLoc="%BaseFolder:"=%\Engine\Binaries\Win64\UnrealEditor.exe"
 
-if exist %UE4EditorLoc:"=% goto GetGame
+if exist %UEEditorLoc:"=% goto GetGame
 
-echo Could not locate UE4Editor.exe
+echo Could not locate UnrealEditor.exe
 goto End
 
 
@@ -54,10 +54,10 @@ set PreDumpCmdLine=%GameName:"=% %DebugDumpParms% %OutputDir:"=% %CaptureDir:"=%
 set PostDumpCmdLine=-forcelogflush
 
 echo Executing dictionary generation commandlet - commandline:
-echo %UE4EditorLoc:"=% %PreDumpCmdLine:"=% %PostDumpCmdLine:"=%
+echo %UEEditorLoc:"=% %PreDumpCmdLine:"=% %PostDumpCmdLine:"=%
 
 
-%UE4EditorLoc:"=% %PreDumpCmdLine:"=% %PostDumpCmdLine:"=%
+%UEEditorLoc:"=% %PreDumpCmdLine:"=% %PostDumpCmdLine:"=%
 
 echo.
 
diff --git a/Engine/Plugins/Runtime/PacketHandlers/CompressionComponents/Oodle/BuildDictionary.bat b/Engine/Plugins/Runtime/PacketHandlers/CompressionComponents/Oodle/BuildDictionary.bat
--- a/Engine/Plugins/Runtime/PacketHandlers/CompressionComponents/Oodle/BuildDictionary.bat
+++ b/Engine/Plugins/Runtime/PacketHandlers/CompressionComponents/Oodle/BuildDictionary.bat
@@ -7,18 +7,18 @@ echo.
 REM This batch file can only work from within the Oodle folder.
 set BaseFolder="..\..\..\..\..\.."
 
-if exist %BaseFolder:"=%\Engine goto SetUE4Editor
+if exist %BaseFolder:"=%\Engine goto SetUEEditor
 
-echo Could not locate Engine folder. This .bat must be located within UE4\Engine\Plugins\Runtime\PacketHandlers\CompressionComponents\Oodle
+echo Could not locate Engine folder. This .bat must be located within UE\Engine\Plugins\Runtime\PacketHandlers\CompressionComponents\Oodle
 goto End
 
 
-:SetUE4Editor
-set UE4EditorLoc="%BaseFolder:"=%\Engine\Binaries\Win64\UE4Editor.exe"
+:SetUEEditor
+set UEEditorLoc="%BaseFolder:"=%\Engine\Binaries\Win64\UnrealEditor.exe"
 
-if exist %UE4EditorLoc:"=% goto GetGame
+if exist %UEEditorLoc:"=% goto GetGame
 
-echo Could not locate UE4Editor.exe
+echo Could not locate UnrealEditor.exe
 goto End
 
 
@@ -36,7 +36,7 @@ echo Executing dictionary generation commandlet - commandline:
 echo %FinalGenCmdLine%
 
 @echo on
-%UE4EditorLoc:"=% %FinalGenCmdLine%
+%UEEditorLoc:"=% %FinalGenCmdLine%
 @echo off
 echo.
 
diff --git a/Engine/Plugins/Runtime/PacketHandlers/CompressionComponents/Oodle/BuildDictionary_Generic.bat b/Engine/Plugins/Runtime/PacketHandlers/CompressionComponents/Oodle/BuildDictionary_Generic.bat
--- a/Engine/Plugins/Runtime/PacketHandlers/CompressionComponents/Oodle/BuildDictionary_Generic.bat
+++ b/Engine/Plugins/Runtime/PacketHandlers/CompressionComponents/Oodle/BuildDictionary_Generic.bat
@@ -7,18 +7,18 @@ echo.
 REM This batch file can only work from within the Oodle folder.
 set BaseFolder="..\..\..\..\..\.."
 
-if exist %BaseFolder:"=%\Engine goto SetUE4Editor
+if exist %BaseFolder:"=%\Engine goto SetUEEditor
 
-echo Could not locate Engine folder. This .bat must be located within UE4\Engine\Plugins\Runtime\PacketHandlers\CompressionComponents\Oodle
+echo Could not locate Engine folder. This .bat must be located within UE\Engine\Plugins\Runtime\PacketHandlers\CompressionComponents\Oodle
 goto End
 
 
-:SetUE4Editor
-set UE4EditorLoc="%BaseFolder:"=%\Engine\Binaries\Win64\UE4Editor.exe"
+:SetUEEditor
+set UEEditorLoc="%BaseFolder:"=%\Engine\Binaries\Win64\UnrealEditor.exe"
 
-if exist %UE4EditorLoc:"=% goto GetGame
+if exist %UEEditorLoc:"=% goto GetGame
 
-echo Could not locate UE4Editor.exe
+echo Could not locate UnrealEditor.exe
 goto End
 
 
@@ -60,7 +60,7 @@ echo Executing dictionary generation commandlet - commandline:
 echo %FinalGenCmdLine%
 
 @echo on
-%UE4EditorLoc:"=% %FinalGenCmdLine%
+%UEEditorLoc:"=% %FinalGenCmdLine%
 @echo off
 echo.
 
diff --git a/Engine/Plugins/Runtime/PacketHandlers/CompressionComponents/Oodle/BuildDictionary_TestCompression.bat b/Engine/Plugins/Runtime/PacketHandlers/CompressionComponents/Oodle/BuildDictionary_TestCompression.bat
--- a/Engine/Plugins/Runtime/PacketHandlers/CompressionComponents/Oodle/BuildDictionary_TestCompression.bat
+++ b/Engine/Plugins/Runtime/PacketHandlers/CompressionComponents/Oodle/BuildDictionary_TestCompression.bat
@@ -7,18 +7,18 @@ echo.
 REM This batch file can only work from within the Oodle folder.
 set BaseFolder="..\..\..\..\..\.."
 
-if exist %BaseFolder:"=%\Engine goto SetUE4Editor
+if exist %BaseFolder:"=%\Engine goto SetUEEditor
 
-echo Could not locate Engine folder. This .bat must be located within UE4\Engine\Plugins\Runtime\PacketHandlers\CompressionComponents\Oodle
+echo Could not locate Engine folder. This .bat must be located within UE\Engine\Plugins\Runtime\PacketHandlers\CompressionComponents\Oodle
 goto End
 
 
-:SetUE4Editor
-set UE4EditorLoc="%BaseFolder:"=%\Engine\Binaries\Win64\UE4Editor.exe"
+:SetUEEditor
+set UEEditorLoc="%BaseFolder:"=%\Engine\Binaries\Win64\UnrealEditor.exe"
 
-if exist %UE4EditorLoc:"=% goto GetGame
+if exist %UEEditorLoc:"=% goto GetGame
 
-echo Could not locate UE4Editor.exe
+echo Could not locate UnrealEditor.exe
 goto End
 
 
@@ -36,7 +36,7 @@ echo Executing dictionary generation commandlet - commandline:
 echo %FinalGenCmdLine%
 
 @echo on
-%UE4EditorLoc:"=% %FinalGenCmdLine%
+%UEEditorLoc:"=% %FinalGenCmdLine%
 @echo off
 echo.
 
diff --git a/Engine/Plugins/Runtime/PacketHandlers/CompressionComponents/Oodle/DumpOodleSDKPackets.bat b/Engine/Plugins/Runtime/PacketHandlers/CompressionComponents/Oodle/DumpOodleSDKPackets.bat
--- a/Engine/Plugins/Runtime/PacketHandlers/CompressionComponents/Oodle/DumpOodleSDKPackets.bat
+++ b/Engine/Plugins/Runtime/PacketHandlers/CompressionComponents/Oodle/DumpOodleSDKPackets.bat
@@ -7,18 +7,18 @@ echo.
 REM This batch file can only work from within the Oodle folder.
 set BaseFolder="..\..\..\..\..\.."
 
-if exist %BaseFolder:"=%\Engine goto SetUE4Editor
+if exist %BaseFolder:"=%\Engine goto SetUEEditor
 
-echo Could not locate Engine folder. This .bat must be located within UE4\Engine\Plugins\Runtime\PacketHandlers\CompressionComponents\Oodle
+echo Could not locate Engine folder. This .bat must be located within UE\Engine\Plugins\Runtime\PacketHandlers\CompressionComponents\Oodle
 goto End
 
 
-:SetUE4Editor
-set UE4EditorLoc="%BaseFolder:"=%\Engine\Binaries\Win64\UE4Editor.exe"
+:SetUEEditor
+set UEEditorLoc="%BaseFolder:"=%\Engine\Binaries\Win64\UnrealEditor.exe"
 
-if exist %UE4EditorLoc:"=% goto GetGame
+if exist %UEEditorLoc:"=% goto GetGame
 
-echo Could not locate UE4Editor.exe
+echo Could not locate UnrealEditor.exe
 goto End
 
 
@@ -36,7 +36,7 @@ echo Executing dictionary generation commandlet - commandline:
 echo %FinalGenCmdLine%
 
 @echo on
-%UE4EditorLoc:"=% %FinalGenCmdLine%
+%UEEditorLoc:"=% %FinalGenCmdLine%
 @echo off
 echo.
 
diff --git a/Engine/Plugins/Runtime/PacketHandlers/CompressionComponents/Oodle/FirstRun.bat b/Engine/Plugins/Runtime/PacketHandlers/CompressionComponents/Oodle/FirstRun.bat
--- a/Engine/Plugins/Runtime/PacketHandlers/CompressionComponents/Oodle/FirstRun.bat
+++ b/Engine/Plugins/Runtime/PacketHandlers/CompressionComponents/Oodle/FirstRun.bat
@@ -7,18 +7,18 @@ echo.
 REM This batch file can only work from within the Oodle folder.
 set BaseFolder="..\..\..\..\..\.."
 
-if exist %BaseFolder:"=%\Engine goto SetUE4Editor
+if exist %BaseFolder:"=%\Engine goto SetUEEditor
 
-echo Could not locate Engine folder. This .bat must be located within UE4\Engine\Plugins\Runtime\PacketHandlers\CompressionComponents\Oodle
+echo Could not locate Engine folder. This .bat must be located within UE\Engine\Plugins\Runtime\PacketHandlers\CompressionComponents\Oodle
 goto End
 
 
-:SetUE4Editor
-set UE4EditorLoc="%BaseFolder:"=%\Engine\Binaries\Win64\UE4Editor.exe"
+:SetUEEditor
+set UEEditorLoc="%BaseFolder:"=%\Engine\Binaries\Win64\UnrealEditor.exe"
 
-if exist %UE4EditorLoc:"=% goto GetGame
+if exist %UEEditorLoc:"=% goto GetGame
 
-echo Could not locate UE4Editor.exe
+echo Could not locate UEEditor.exe
 goto End
 
 
@@ -36,7 +36,7 @@ echo Executing plugin enable commandlet - commandline:
 echo %FinalEnableCmdLine%
 
 @echo on
-%UE4EditorLoc:"=% %FinalEnableCmdLine%
+%UEEditorLoc:"=% %FinalEnableCmdLine%
 @echo off
 echo.
 
@@ -56,7 +56,7 @@ echo Executing Oodle PacketHandler enable commandlet - commandline:
 echo %FinalHandlerCmdLine%
 
 @echo on
-%UE4EditorLoc:"=% %FinalHandlerCmdLine%
+%UEEditorLoc:"=% %FinalHandlerCmdLine%
 @echo off
 echo.
 
diff --git a/Engine/Plugins/Runtime/PacketHandlers/CompressionComponents/Oodle/Source/OodleCompressionFormat/Private/OodleCompressionFormat.cpp b/Engine/Plugins/Runtime/PacketHandlers/CompressionComponents/Oodle/Source/OodleCompressionFormat/Private/OodleCompressionFormat.cpp
--- a/Engine/Plugins/Runtime/PacketHandlers/CompressionComponents/Oodle/Source/OodleCompressionFormat/Private/OodleCompressionFormat.cpp
+++ b/Engine/Plugins/Runtime/PacketHandlers/CompressionComponents/Oodle/Source/OodleCompressionFormat/Private/OodleCompressionFormat.cpp
@@ -57,7 +57,7 @@ struct FOodleCustomCompressor : ICompressionFormat
 			return TEXT("Kraken");
 		case OodleLZ_Compressor_Selkie:
 			return TEXT("Selkie");
-#if UE4_OODLE_VER >= 270
+#if UE_OODLE_VER >= 270
 		case OodleLZ_Compressor_Leviathan:
 			return TEXT("Leviathan");
 #else
@@ -101,7 +101,7 @@ struct FOodleCustomCompressor : ICompressionFormat
 
 	virtual uint32 GetVersion()
 	{
-		return uint32(UE4_OODLE_VER);
+		return uint32(UE_OODLE_VER);
 	}
 
 	virtual FString GetDDCKeySuffix()
@@ -145,7 +145,7 @@ struct FOodleCustomCompressor : ICompressionFormat
 	{
 		ConditionalInitialize();
 
-#if UE4_OODLE_VER > 255
+#if UE_OODLE_VER > 255
 		int32 Needed = (int32)OodleLZ_GetCompressedBufferSizeNeeded(Compressor, UncompressedSize);
 #else
 		int32 Needed = (int32)OodleLZ_GetCompressedBufferSizeNeeded(UncompressedSize);
@@ -178,7 +178,7 @@ class FOodleCompressionFormatModuleInterface : public IModuleInterface
 			{ TEXT("Mermaid"), OodleLZ_Compressor_Mermaid },
 			{ TEXT("Kraken"), OodleLZ_Compressor_Kraken },
 			{ TEXT("Selkie"), OodleLZ_Compressor_Selkie },
-#if UE4_OODLE_VER >= 270
+#if UE_OODLE_VER >= 270
 			{ TEXT("Leviathan"), OodleLZ_Compressor_Leviathan },
 #else
 			{ TEXT("LZNA"), OodleLZ_Compressor_LZNA },
@@ -207,7 +207,7 @@ class FOodleCompressionFormatModuleInterface : public IModuleInterface
 
 		IModularFeatures::Get().RegisterModularFeature(COMPRESSION_FORMAT_FEATURE_NAME, CompressionFormat);
 
-#if UE4_OODLE_VER > 255
+#if UE_OODLE_VER > 255
 		OodleCore_Plugins_SetPrintf(OodleCore_Plugin_Printf_Verbose);
 #endif
 
diff --git a/Engine/Plugins/Runtime/PacketHandlers/CompressionComponents/Oodle/Source/OodleHandlerComponent/Classes/OodleTrainerCommandlet.h b/Engine/Plugins/Runtime/PacketHandlers/CompressionComponents/Oodle/Source/OodleHandlerComponent/Classes/OodleTrainerCommandlet.h
--- a/Engine/Plugins/Runtime/PacketHandlers/CompressionComponents/Oodle/Source/OodleHandlerComponent/Classes/OodleTrainerCommandlet.h
+++ b/Engine/Plugins/Runtime/PacketHandlers/CompressionComponents/Oodle/Source/OodleHandlerComponent/Classes/OodleTrainerCommandlet.h
@@ -13,7 +13,7 @@
 #define USE_OODLE_TRAINER_COMMANDLET (HAS_OODLE_NET_SDK && WITH_EDITOR)
 
 /**
- * Commandlet for processing UE4 packet captures, through Oodle's training API, for generating compressed state dictionaries.
+ * Commandlet for processing UE packet captures, through Oodle's training API, for generating compressed state dictionaries.
  *
  *
  * Primary Commands:
diff --git a/Engine/Plugins/Runtime/PacketHandlers/CompressionComponents/Oodle/Source/OodleHandlerComponent/Private/OodleArchives.cpp b/Engine/Plugins/Runtime/PacketHandlers/CompressionComponents/Oodle/Source/OodleHandlerComponent/Private/OodleArchives.cpp
--- a/Engine/Plugins/Runtime/PacketHandlers/CompressionComponents/Oodle/Source/OodleHandlerComponent/Private/OodleArchives.cpp
+++ b/Engine/Plugins/Runtime/PacketHandlers/CompressionComponents/Oodle/Source/OodleHandlerComponent/Private/OodleArchives.cpp
@@ -43,7 +43,7 @@ bool FOodleArchiveBase::SerializeOodleCompressData(FOodleCompressedData& OutData
 		OutDataInfo.DecompressedLength.Set(*this, DataBytes);
 	
 #if HAS_OODLE_DATA_SDK
-#if UE4_OODLE_VER > 255
+#if UE_OODLE_VER > 255
 		uint32 CompressBufferLen = OodleLZ_GetCompressedBufferSizeNeeded(OodleLZ_Compressor_Kraken, DataBytes);
 #else
 		uint32 CompressBufferLen = OodleLZ_GetCompressedBufferSizeNeeded(DataBytes);
@@ -123,7 +123,7 @@ bool FOodleArchiveBase::SerializeOodleDecompressData(FOodleCompressedData& DataI
 			}
 	
 			InnerArchive.Serialize(CompressedData, CompressedLength);
-#if UE4_OODLE_VER > 255
+#if UE_OODLE_VER > 255
 			SINTa OodleLen = OodleLZ_Decompress((void*)CompressedData, CompressedLength, (void*)DecompressedData, DecompressedLength, OodleLZ_FuzzSafe_Yes);
 #else
 			SINTa OodleLen = OodleLZ_Decompress((void*)CompressedData, CompressedLength, (void*)DecompressedData, DecompressedLength, OodleLZ_FuzzSafe_No);
@@ -399,7 +399,7 @@ FOodleDictionaryArchive::FOodleDictionaryArchive(FArchive& InInnerArchive)
 	if (IsSaving())
 	{
 		Header.DictionaryVersion = DICTIONARY_FILE_VERSION;
-#if UE4_OODLE_VER > 255
+#if UE_OODLE_VER > 255
 		Header.OodleMajorHeaderVersion = OODLE2NET_VERSION_MAJOR;
 #else
 		Header.OodleMajorHeaderVersion = OODLE2_VERSION_MAJOR;
diff --git a/Engine/Plugins/Runtime/PacketHandlers/CompressionComponents/Oodle/Source/OodleHandlerComponent/Private/OodleUtils.cpp b/Engine/Plugins/Runtime/PacketHandlers/CompressionComponents/Oodle/Source/OodleHandlerComponent/Private/OodleUtils.cpp
--- a/Engine/Plugins/Runtime/PacketHandlers/CompressionComponents/Oodle/Source/OodleHandlerComponent/Private/OodleUtils.cpp
+++ b/Engine/Plugins/Runtime/PacketHandlers/CompressionComponents/Oodle/Source/OodleHandlerComponent/Private/OodleUtils.cpp
@@ -47,7 +47,7 @@ namespace OodleUtils
 		const OodleLZ_Compressor			Compressor = OodleLZ_Compressor_Selkie;
 		const OodleLZ_CompressionLevel		Level = OodleLZ_CompressionLevel_VeryFast;
 
-#if UE4_OODLE_VER > 255
+#if UE_OODLE_VER > 255
 		OutCompressed.SetNum(ReservedBytes + OodleLZ_GetCompressedBufferSizeNeeded(Compressor, InBuffer.Num()));
 #else
 		OutCompressed.SetNum(ReservedBytes + OodleLZ_GetCompressedBufferSizeNeeded(InBuffer.Num()));
diff --git a/Engine/Plugins/Runtime/PacketHandlers/CompressionComponents/Oodle/Source/OodleHandlerComponent/Public/OodleArchives.h b/Engine/Plugins/Runtime/PacketHandlers/CompressionComponents/Oodle/Source/OodleHandlerComponent/Public/OodleArchives.h
--- a/Engine/Plugins/Runtime/PacketHandlers/CompressionComponents/Oodle/Source/OodleHandlerComponent/Public/OodleArchives.h
+++ b/Engine/Plugins/Runtime/PacketHandlers/CompressionComponents/Oodle/Source/OodleHandlerComponent/Public/OodleArchives.h
@@ -362,7 +362,7 @@ public:
 
 
 /**
- * Archive for handling UE4 Oodle dictionary (.udic) files
+ * Archive for handling UE Oodle dictionary (.udic) files
  */
 class FOodleDictionaryArchive : public FOodleArchiveBase
 {
diff --git a/Engine/Plugins/Runtime/PacketHandlers/CompressionComponents/Oodle/Source/OodleHandlerComponent/Public/OodleHandlerComponent.h b/Engine/Plugins/Runtime/PacketHandlers/CompressionComponents/Oodle/Source/OodleHandlerComponent/Public/OodleHandlerComponent.h
--- a/Engine/Plugins/Runtime/PacketHandlers/CompressionComponents/Oodle/Source/OodleHandlerComponent/Public/OodleHandlerComponent.h
+++ b/Engine/Plugins/Runtime/PacketHandlers/CompressionComponents/Oodle/Source/OodleHandlerComponent/Public/OodleHandlerComponent.h
@@ -35,12 +35,12 @@ enum class EOodleEnableMode : uint8
 
 #if HAS_OODLE_NET_SDK
 
-#if UE4_OODLE_VER >= 270
+#if UE_OODLE_VER >= 270
 #include "oodle2net.h"
 	#if HAS_OODLE_DATA_SDK // allow data SDK usage if we have it
 	#include "oodle2.h"
 	#endif
-#elif UE4_OODLE_VER >= 200
+#elif UE_OODLE_VER >= 200
 #include "oodle2.h"
 #else
 #include "oodle.h"
